Fixed warning:  Engine DDC Win64 & Linux - Unable to find Slate Widget Style 'SpinBox'. The fix was done in the experimental Chaos Editor. The error occurred because it was forgotten to call the FEditorStyle:: when calling GetWidgetStyle<FSpinBoxStyle>.

diff --git a/Engine/Plugins/Experimental/ChaosEditor/Source/FractureEditor/Private/FractureEditorStyle.cpp b/Engine/Plugins/Experimental/ChaosEditor/Source/FractureEditor/Private/FractureEditorStyle.cpp
--- a/Engine/Plugins/Experimental/ChaosEditor/Source/FractureEditor/Private/FractureEditorStyle.cpp
+++ b/Engine/Plugins/Experimental/ChaosEditor/Source/FractureEditor/Private/FractureEditorStyle.cpp
@@ -76,7 +76,7 @@ FFractureEditorStyle::FFractureEditorStyle()
 	Set("FractureEditor.ViewDownOneLevel", new FSlateImageBrush(RootToContentDir(TEXT("LevelViewDown_48x.png")), IconSize));
 	Set("FractureEditor.ViewDownOneLevel.Small", new FSlateImageBrush(RootToContentDir(TEXT("LevelViewDown_48x.png")), SmallIconSize));
 
-	Set("FractureEditor.SpinBox", FSpinBoxStyle(GetWidgetStyle<FSpinBoxStyle>("SpinBox"))
+	Set("FractureEditor.SpinBox", FSpinBoxStyle(FEditorStyle::GetWidgetStyle<FSpinBoxStyle>("SpinBox"))
 		.SetTextPadding(FMargin(0))
 		.SetBackgroundBrush(FSlateNoResource())
 		.SetHoveredBackgroundBrush(FSlateNoResource())
