FFractureEditorModeToolkit::UpdateExplodedVectors() "ExplodedVector" inits so 0,0,0 on change.  #jira UE-75787 #jira UE-76112
FFractureEditorModeToolkit::UpdateExplodedVectors() now explodes all bones if view level is -1.  #jira UE-75837  #jira UE-75730
Outliner changes now calls ActiveTool->FractureContextChanged() #jira UE-75819
Flattening GC sets view level to 1 noew #jira UE-75903
Selecting a GC actor now updates it to current explosion level #jira UE-75985
OnFractureClicked() now invalidates VP to fix hitproxy data #jira UE-75807
Fracture button enable state now disable if no fracturable leaf is selected.  Also, non-fracturable bones indices are no longer passed to cutter #jira UE-75823
OnCluster will now reset highlighting, update tree view with selection. #jira UE-75818
Modify GeometryCollection just prior to fracture #jira UE-75926
Fix variation in radial fracture #jira UETOOL-1683

diff --git a/Engine/Plugins/Experimental/ChaosEditor/Source/FractureEditor/Private/FractureEditorMode.cpp b/Engine/Plugins/Experimental/ChaosEditor/Source/FractureEditor/Private/FractureEditorMode.cpp
--- a/Engine/Plugins/Experimental/ChaosEditor/Source/FractureEditor/Private/FractureEditorMode.cpp
+++ b/Engine/Plugins/Experimental/ChaosEditor/Source/FractureEditor/Private/FractureEditorMode.cpp
@@ -369,7 +369,7 @@ void FFractureEditorMode::GetActorGlobalBounds(TArrayView<UGeometryCollectionCom
 				{
 					const FVector& Offset = ExplodedVectors[Idx];
 					const FBox& Bounds = BoundingBox[TransformToGeometryIndex[Idx]];
-					BoundsToBone.Add(TransformToGeometryIndex[Idx], Bounds.ShiftBy(Offset).TransformBy(Transforms[Idx]));
+					BoundsToBone.Add(Idx, Bounds.ShiftBy(Offset).TransformBy(Transforms[Idx]));
 				}
 			}
 		}
@@ -379,7 +379,7 @@ void FFractureEditorMode::GetActorGlobalBounds(TArrayView<UGeometryCollectionCom
 			{
 				if (TransformToGeometryIndex[Idx] > -1)
 				{
-					BoundsToBone.Add(TransformToGeometryIndex[Idx], BoundingBox[TransformToGeometryIndex[Idx]].TransformBy(Transforms[Idx]));
+					BoundsToBone.Add(Idx, BoundingBox[TransformToGeometryIndex[Idx]].TransformBy(Transforms[Idx]));
 				}
 			}
 		}
diff --git a/Engine/Plugins/Experimental/ChaosEditor/Source/FractureEditor/Private/FractureEditorModeToolkit.cpp b/Engine/Plugins/Experimental/ChaosEditor/Source/FractureEditor/Private/FractureEditorModeToolkit.cpp
--- a/Engine/Plugins/Experimental/ChaosEditor/Source/FractureEditor/Private/FractureEditorModeToolkit.cpp
+++ b/Engine/Plugins/Experimental/ChaosEditor/Source/FractureEditor/Private/FractureEditorModeToolkit.cpp
@@ -47,6 +47,8 @@
 #include "FractureSelectionTools.h"
 #include "GeometryCollection/GeometryCollectionAlgo.h"
 #include "GeometryCollection/GeometryCollectionClusteringUtility.h"
+#include "Editor.h"
+#include "LevelEditorViewport.h"
 
 #define LOCTEXT_NAMESPACE "FFractureEditorModeToolkit"
 
@@ -976,6 +978,7 @@ void FFractureEditorModeToolkit::SetOutlinerComponents(const TArray<UGeometryCol
 		TSharedPtr<FGeometryCollection, ESPMode::ThreadSafe> GeometryCollectionPtr = FracturedGeometryCollection->GetGeometryCollection();
 
 		FGeometryCollectionClusteringUtility::UpdateHierarchyLevelOfChildren(GeometryCollectionPtr.Get(), -1);
+		UpdateExplodedVectors(Component);
 	}
 
 	if (OutlinerView)
@@ -1066,25 +1069,40 @@ void FFractureEditorModeToolkit::GetFractureContexts(TArray<FFractureContext>& F
 					UGeometryCollection* FracturedGeometryCollection = RestCollection.GetRestCollection();
 
 					FScopedColorEdit EditBoneColor = GeometryCollectionComponent->EditBoneSelection();
-					const TArray<int32>& SelectedBones = EditBoneColor.GetSelectedBones();
-					const TManagedArray<FBox>& BoundingBoxes = GeometryCollectionComponent->GetBoundingBoxArray();
+					const TArray<int32>& SelectedBonesOriginal = EditBoneColor.GetSelectedBones();
 
 					TSharedPtr<FGeometryCollection, ESPMode::ThreadSafe> GeometryCollectionPtr = FracturedGeometryCollection->GetGeometryCollection();
 					FGeometryCollection* OutGeometryCollection = GeometryCollectionPtr.Get();
 
+					const TManagedArray<TSet<int32>>& Children = OutGeometryCollection->GetAttribute<TSet<int32>>("Children", FGeometryCollection::TransformGroup);
+
+					TArray<int32> SelectedBones;
+					SelectedBones.Reserve(SelectedBonesOriginal.Num());
+					for (int32 BoneIndex : SelectedBonesOriginal)
+					{
+						if (Children[BoneIndex].Num() == 0)
+						{
+							SelectedBones.Add(BoneIndex);
+						}
+					}
+
 					const TManagedArray<FTransform>& Transform = OutGeometryCollection->GetAttribute<FTransform>("Transform", FGeometryCollection::TransformGroup);
 					const TManagedArray<int32>& TransformToGeometryIndex = OutGeometryCollection->GetAttribute<int32>("TransformToGeometryIndex", FGeometryCollection::TransformGroup);
 
+					const TManagedArray<FBox>& BoundingBoxes = OutGeometryCollection->GetAttribute<FBox>("BoundingBox", FGeometryCollection::GeometryGroup);
+
 					TArray<FTransform> Transforms;
 					GeometryCollectionAlgo::GlobalMatrices(Transform, OutGeometryCollection->Parent, Transforms);
 
 					TMap<int32, FBox> BoundsToBone;
 
+
 					for (int32 Idx = 0, ni = FracturedGeometryCollection->NumElements(FGeometryCollection::TransformGroup); Idx < ni; ++Idx)
 					{
 						if (TransformToGeometryIndex[Idx] > -1)
 						{
-							BoundsToBone.Add(TransformToGeometryIndex[Idx], BoundingBoxes[TransformToGeometryIndex[Idx]].TransformBy(Transforms[Idx]));
+							ensure(TransformToGeometryIndex[Idx] > -1);
+							BoundsToBone.Add(Idx, BoundingBoxes[TransformToGeometryIndex[Idx]].TransformBy(Transforms[Idx]));
 						}
 					}
 
@@ -1108,7 +1126,10 @@ void FFractureEditorModeToolkit::GetFractureContexts(TArray<FFractureContext>& F
 						FractureContext.Bounds = FBox(ForceInit);
 						for (int32 BoneIndex : FractureContext.SelectedBones)
 						{
-							FractureContext.Bounds += BoundsToBone[TransformToGeometryIndex[BoneIndex]];
+ 							if(TransformToGeometryIndex[BoneIndex] > -1)
+							{
+								FractureContext.Bounds += BoundsToBone[BoneIndex];
+							}
 						}
 					}
 					else
@@ -1129,7 +1150,10 @@ void FFractureEditorModeToolkit::GetFractureContexts(TArray<FFractureContext>& F
 							FractureContext.OriginalPrimitiveComponent = GeometryCollectionComponent;
 							FractureContext.FracturedGeometryCollection = FracturedGeometryCollection;
 							FractureContext.SelectedBones = { BoneIndex };
-							FractureContext.Bounds = BoundsToBone[TransformToGeometryIndex[BoneIndex]];
+							if(TransformToGeometryIndex[BoneIndex] > -1)
+							{
+								FractureContext.Bounds = BoundsToBone[BoneIndex];
+							}
 						}
 					}
 				}
@@ -1158,6 +1182,11 @@ FReply FFractureEditorModeToolkit::OnFractureClicked()
 		
 			UpdateExplodedVectors(GeometryCollectionComponent);
 			SetOutlinerComponents({ GeometryCollectionComponent });
+
+			GeometryCollectionComponent->MarkRenderDynamicDataDirty();
+			GeometryCollectionComponent->MarkRenderStateDirty();
+
+			GCurrentLevelEditingViewportClient->Invalidate();
 		}
 
 		float ProcessingTime = static_cast<float>(FPlatformTime::Seconds() - CacheStartTime);
@@ -1175,7 +1204,7 @@ bool FFractureEditorModeToolkit::CanExecuteFracture() const
 
 	if (IsGeometryCollectionSelected())
 	{
-		return true;
+		return IsLeafBoneSelected();
 	}
 
 	if (IsStaticMeshSelected())
@@ -1187,6 +1216,38 @@ bool FFractureEditorModeToolkit::CanExecuteFracture() const
 }
 
 
+bool FFractureEditorModeToolkit::IsLeafBoneSelected() const
+{
+	TSet<UGeometryCollectionComponent*> GeomCompSelection;
+	GetSelectedGeometryCollectionComponents(GeomCompSelection);
+	for (UGeometryCollectionComponent* GeometryCollectionComponent : GeomCompSelection)
+	{
+		const TArray<int32> SelectedBones = GeometryCollectionComponent->GetSelectedBones();
+
+		if (SelectedBones.Num() > 0)
+		{
+			FGeometryCollectionEdit GCEdit = GeometryCollectionComponent->EditRestCollection();
+			if (UGeometryCollection* GCObject = GCEdit.GetRestCollection())
+			{
+				TSharedPtr<FGeometryCollection, ESPMode::ThreadSafe> GeometryCollectionPtr = GCObject->GetGeometryCollection();
+				if (FGeometryCollection* GeometryCollection = GeometryCollectionPtr.Get())
+				{
+					const TManagedArray<TSet<int32>>& Children = GeometryCollection->GetAttribute<TSet<int32>>("Children", FGeometryCollection::TransformGroup);
+
+					for (int32 BoneIndex : SelectedBones)
+					{
+						if (Children[BoneIndex].Num() == 0)
+						{
+							return true;
+						}
+					}
+				}
+			}
+		}
+	}
+	return false;
+}
+
 ULevel* FFractureEditorModeToolkit::GetSelectedLevel()
 {
 	USelection* SelectedActors = GEditor->GetSelectedActors();
@@ -1280,9 +1341,11 @@ void FFractureEditorModeToolkit::OnCluster()
 					FGeometryCollectionClusteringUtility::UpdateHierarchyLevelOfChildren(GeometryCollection, -1);
 
 					FScopedColorEdit EditBoneColor = GeometryCollectionComponent->EditBoneSelection();
-					EditBoneColor.ResetBoneSelection();
+ 					EditBoneColor.ResetBoneSelection();
+					EditBoneColor.ResetHighlightedBones();
 					GeometryCollectionComponent->MarkRenderDynamicDataDirty();
 					GeometryCollectionComponent->MarkRenderStateDirty();
+					SetBoneSelection(GeometryCollectionComponent, EditBoneColor.GetSelectedBones(), true);
 				}
 			}
 		}
@@ -1311,8 +1374,10 @@ void FFractureEditorModeToolkit::OnUncluster()
 
 				FScopedColorEdit EditBoneColor = GeometryCollectionComponent->EditBoneSelection();
 				EditBoneColor.ResetBoneSelection();
+				EditBoneColor.ResetHighlightedBones();
 				GeometryCollectionComponent->MarkRenderDynamicDataDirty();
 				GeometryCollectionComponent->MarkRenderStateDirty();
+				SetBoneSelection(GeometryCollectionComponent, EditBoneColor.GetSelectedBones(), true);
 			}
 		}
 	}
@@ -1382,6 +1447,8 @@ void FFractureEditorModeToolkit::OnFlatten()
 				FScopedColorEdit EditBoneColor = GeometryCollectionComponent->EditBoneSelection();
 				EditBoneColor.ResetBoneSelection();
 
+				OnSetLevelViewValue(1);
+
 				GeometryCollectionComponent->MarkRenderDynamicDataDirty();
 				GeometryCollectionComponent->MarkRenderStateDirty();
 			}
@@ -1580,6 +1647,7 @@ class AGeometryCollectionActor* FFractureEditorModeToolkit::CreateNewGeometryAct
 
 void FFractureEditorModeToolkit::ExecuteFracture(FFractureContext& FractureContext)
 {
+	FractureContext.FracturedGeometryCollection->Modify();
 	ActiveTool->ExecuteFracture(FractureContext);
 	TSharedPtr<FGeometryCollection, ESPMode::ThreadSafe> GeometryCollectionPtr = FractureContext.FracturedGeometryCollection->GetGeometryCollection();
 	FGeometryCollection* OutGeometryCollection = GeometryCollectionPtr.Get();
@@ -1722,12 +1790,6 @@ void FFractureEditorModeToolkit::UpdateExplodedVectors(UGeometryCollectionCompon
 				MaxFractureLevel = Levels[Idx];
 		}
 
-		// < 0 means it's showing all leafs,
-		if (ViewFractureLevel < 0)
-		{
-			ViewFractureLevel = MaxFractureLevel;
-		}
-
 		TArray<FTransform> Transforms;
 		GeometryCollectionAlgo::GlobalMatrices(Transform, OutGeometryCollection->Parent, Transforms);
 
@@ -1739,11 +1801,12 @@ void FFractureEditorModeToolkit::UpdateExplodedVectors(UGeometryCollectionCompon
 		FVector Center(ForceInitToZero);
 		for (int32 Idx = 0, ni = GeometryCollection->NumElements(FGeometryCollection::TransformGroup); Idx < ni; ++Idx)
 		{
+			ExplodedVectors[Idx] = FVector::ZeroVector;
 			FVector GeoCenter;
 			if (GetValidGeoCenter(TransformToGeometryIndex, Transforms, Children, BoundingBox, Idx, GeoCenter))
 			{
 				TransformedCenters[Idx] = GeoCenter;
-				if (Levels[Idx] == ViewFractureLevel)
+				if ((ViewFractureLevel < 0) || Levels[Idx] == ViewFractureLevel)
 				{
 					Center += TransformedCenters[Idx];
 					++TransformsCount;
@@ -1757,7 +1820,7 @@ void FFractureEditorModeToolkit::UpdateExplodedVectors(UGeometryCollectionCompon
 		{
 			for (int32 Idx = 0, ni = GeometryCollection->NumElements(FGeometryCollection::TransformGroup); Idx < ni; ++Idx)
 			{
-				if (Levels[Idx] == ViewFractureLevel)
+				if ((ViewFractureLevel < 0) || Levels[Idx] == ViewFractureLevel)
 				{
 					ExplodedVectors[Idx] = (TransformedCenters[Idx] - Center) * ExplodeAmount;
 				}
@@ -1821,6 +1884,11 @@ void FFractureEditorModeToolkit::OnOutlinerBoneSelectionChanged(UGeometryCollect
 		FFractureSelectionTools::ClearSelectedBones(RootComponent);
 	}
 
+	if (ActiveTool != nullptr)
+	{
+		ActiveTool->FractureContextChanged();
+	}
+
 	RootComponent->MarkRenderStateDirty();
 	RootComponent->MarkRenderDynamicDataDirty();
 }
diff --git a/Engine/Plugins/Experimental/ChaosEditor/Source/FractureEditor/Private/FractureToolRadial.cpp b/Engine/Plugins/Experimental/ChaosEditor/Source/FractureEditor/Private/FractureToolRadial.cpp
--- a/Engine/Plugins/Experimental/ChaosEditor/Source/FractureEditor/Private/FractureToolRadial.cpp
+++ b/Engine/Plugins/Experimental/ChaosEditor/Source/FractureEditor/Private/FractureToolRadial.cpp
@@ -77,7 +77,9 @@ void UFractureToolRadial::GenerateVoronoiSites(const FFractureContext &Context,
 
 	const FVector Center(Context.Bounds.GetCenter());
 
-	for (int32 ii = 0; ii < FractureSettings->RadialSteps; ++ii)
+	FRandomStream RandStream(Context.RandomSeed);
+
+	for (int32 ii = 1; ii < FractureSettings->RadialSteps; ++ii)
 	{
 		FVector PositionVector(RadialStep * ii,0.0,0.0);
 		float AngularStep = 360.f / FractureSettings->AngularSteps;
@@ -86,7 +88,7 @@ void UFractureToolRadial::GenerateVoronoiSites(const FFractureContext &Context,
 		for (int32 kk = 0; kk < FractureSettings->AngularSteps; ++kk)
 		{
 			PositionVector = PositionVector.RotateAngleAxis(AngularStep , FractureSettings->Normal);
-			Sites.Emplace(Center + PositionVector);
+			Sites.Emplace(Center + PositionVector + (RandStream.VRand() * RandStream.FRand() * FractureSettings->Variability));
 		}
 	}
 }
diff --git a/Engine/Plugins/Experimental/ChaosEditor/Source/FractureEditor/Private/FractureToolRadial.h b/Engine/Plugins/Experimental/ChaosEditor/Source/FractureEditor/Private/FractureToolRadial.h
--- a/Engine/Plugins/Experimental/ChaosEditor/Source/FractureEditor/Private/FractureToolRadial.h
+++ b/Engine/Plugins/Experimental/ChaosEditor/Source/FractureEditor/Private/FractureToolRadial.h
@@ -38,23 +38,23 @@ public:
 	FVector Normal;
 
 	/** Pattern radius */
-	UPROPERTY(EditAnywhere, Category = RadialVoronoi)
+	UPROPERTY(EditAnywhere, Category = RadialVoronoi, meta = (DisplayName = "Radius", UIMin = "0.0", ClampMin = "0.0"))
 	float Radius;
 
 	/** Number of angular steps */
-	UPROPERTY(EditAnywhere, Category = RadialVoronoi)
+	UPROPERTY(EditAnywhere, Category = RadialVoronoi, meta = (DisplayName = "Angular Steps", UIMin = "1", ClampMin = "1"))
 	int AngularSteps;
 
 	/** Number of radial steps */
-	UPROPERTY(EditAnywhere, Category = RadialVoronoi)
+	UPROPERTY(EditAnywhere, Category = RadialVoronoi, meta = (DisplayName = "Radial Steps", UIMin = "1", ClampMin = "1"))
 	int RadialSteps;
 
 	/** Angle offset at each radial step */
-	UPROPERTY(EditAnywhere, Category = RadialVoronoi)
+	UPROPERTY(EditAnywhere, Category = RadialVoronoi, meta = (DisplayName = "Angle Offset", UIMin = "0.0", ClampMin = "0.0"))
 	float AngleOffset;
 
 	/** Randomness of sites distribution */
-	UPROPERTY(EditAnywhere, Category = RadialVoronoi)
+	UPROPERTY(EditAnywhere, Category = RadialVoronoi, meta = (DisplayName = "Variability", UIMin = "0.0", ClampMin = "0.0"))
 	float Variability;
 
 	UPROPERTY()
