-file missing from shelf #rb none #jira none

diff --git a/Engine/Plugins/Experimental/ChaosEditor/Source/FractureEditor/Public/FractureEditorModeToolkit.h b/Engine/Plugins/Experimental/ChaosEditor/Source/FractureEditor/Public/FractureEditorModeToolkit.h
--- a/Engine/Plugins/Experimental/ChaosEditor/Source/FractureEditor/Public/FractureEditorModeToolkit.h
+++ b/Engine/Plugins/Experimental/ChaosEditor/Source/FractureEditor/Public/FractureEditorModeToolkit.h
@@ -73,6 +73,7 @@ public:
 	// Fracture Command Callback
 	FReply OnFractureClicked();
 	bool CanExecuteFracture() const;
+	bool IsLeafBoneSelected() const;
 
 	static ULevel* GetSelectedLevel();
  	static AActor* AddActor(ULevel* InLevel, UClass* Class);
