Add missing "ChanceToFracture" test for each fracture context. #rb none #jira UE-76236

diff --git a/Engine/Plugins/Experimental/ChaosEditor/Source/FractureEditor/Private/FractureEditorModeToolkit.cpp b/Engine/Plugins/Experimental/ChaosEditor/Source/FractureEditor/Private/FractureEditorModeToolkit.cpp
--- a/Engine/Plugins/Experimental/ChaosEditor/Source/FractureEditor/Private/FractureEditorModeToolkit.cpp
+++ b/Engine/Plugins/Experimental/ChaosEditor/Source/FractureEditor/Private/FractureEditorModeToolkit.cpp
@@ -1052,6 +1052,7 @@ void FFractureEditorModeToolkit::OnSelectByMode(GeometryCollection::ESelectionMo
 void FFractureEditorModeToolkit::GetFractureContexts(TArray<FFractureContext>& FractureContexts)
 {
 	const UFractureCommonSettings* CommonSettings = GetDefault<UFractureCommonSettings>();
+	FRandomStream RandomStream(CommonSettings->RandomSeed > -1 ? CommonSettings->RandomSeed : FMath::Rand());
 
 	USelection* SelectedActors = GEditor->GetSelectedActors();
 	for (FSelectionIterator Iter(*SelectedActors); Iter; ++Iter)
@@ -1096,7 +1097,6 @@ void FFractureEditorModeToolkit::GetFractureContexts(TArray<FFractureContext>& F
 
 					TMap<int32, FBox> BoundsToBone;
 
-
 					for (int32 Idx = 0, ni = FracturedGeometryCollection->NumElements(FGeometryCollection::TransformGroup); Idx < ni; ++Idx)
 					{
 						if (TransformToGeometryIndex[Idx] > -1)
@@ -1126,7 +1126,12 @@ void FFractureEditorModeToolkit::GetFractureContexts(TArray<FFractureContext>& F
 						FractureContext.Bounds = FBox(ForceInit);
 						for (int32 BoneIndex : FractureContext.SelectedBones)
 						{
- 							if(TransformToGeometryIndex[BoneIndex] > -1)
+							if (FractureContext.SelectedBones.Num() > 1 && RandomStream.FRand() > CommonSettings->ChanceToFracture)
+							{
+								continue;
+							}
+							
+							if (TransformToGeometryIndex[BoneIndex] > -1)
 							{
 								FractureContext.Bounds += BoundsToBone[BoneIndex];
 							}
@@ -1136,6 +1141,11 @@ void FFractureEditorModeToolkit::GetFractureContexts(TArray<FFractureContext>& F
 					{
 						for (int32 BoneIndex : SelectedBones)
 						{
+							if (SelectedBones.Num() > 1 && RandomStream.FRand() > CommonSettings->ChanceToFracture)
+							{
+								continue;
+							}
+
 							FractureContexts.AddDefaulted();
 							FFractureContext& FractureContext = FractureContexts.Last();
 							FractureContext.RandomSeed = FMath::Rand();
