Fix to get perpendicular vector to up instead of assuming X #rb none #jira UE-76226

diff --git a/Engine/Plugins/Experimental/ChaosEditor/Source/FractureEditor/Private/FractureToolRadial.cpp b/Engine/Plugins/Experimental/ChaosEditor/Source/FractureEditor/Private/FractureToolRadial.cpp
--- a/Engine/Plugins/Experimental/ChaosEditor/Source/FractureEditor/Private/FractureToolRadial.cpp
+++ b/Engine/Plugins/Experimental/ChaosEditor/Source/FractureEditor/Private/FractureToolRadial.cpp
@@ -69,16 +69,20 @@ void UFractureToolRadial::GenerateVoronoiSites(const FFractureContext &Context,
 	const FVector Center(Context.Bounds.GetCenter() + FractureSettings->Center);
 
 	FRandomStream RandStream(Context.RandomSeed);
+	FVector UpVector(FractureSettings->Normal);
+	UpVector.Normalize();
+	FVector PerpVector(UpVector[2], UpVector[0], UpVector[1]);
 
 	for (int32 ii = 1; ii < FractureSettings->RadialSteps; ++ii)
 	{
-		FVector PositionVector(RadialStep * ii,0.0,0.0);
+		FVector PositionVector(PerpVector * RadialStep * ii);
+
 		float AngularStep = 360.f / FractureSettings->AngularSteps;
-		PositionVector = PositionVector.RotateAngleAxis(FractureSettings->AngleOffset * ii, FractureSettings->Normal);
+		PositionVector = PositionVector.RotateAngleAxis(FractureSettings->AngleOffset * ii, UpVector);
 
 		for (int32 kk = 0; kk < FractureSettings->AngularSteps; ++kk)
 		{
-			PositionVector = PositionVector.RotateAngleAxis(AngularStep , FractureSettings->Normal);
+			PositionVector = PositionVector.RotateAngleAxis(AngularStep , UpVector);
 			Sites.Emplace(Center + PositionVector + (RandStream.VRand() * RandStream.FRand() * FractureSettings->Variability));
 		}
 	}
