Planar fracture with a reference object now corrects for transfomrs offsets between them.  #rb none #jira UE-76180

diff --git a/Engine/Plugins/Experimental/ChaosEditor/Source/FractureEditor/Private/FractureToolPlaneCut.cpp b/Engine/Plugins/Experimental/ChaosEditor/Source/FractureEditor/Private/FractureToolPlaneCut.cpp
--- a/Engine/Plugins/Experimental/ChaosEditor/Source/FractureEditor/Private/FractureToolPlaneCut.cpp
+++ b/Engine/Plugins/Experimental/ChaosEditor/Source/FractureEditor/Private/FractureToolPlaneCut.cpp
@@ -133,7 +133,7 @@ void UFractureToolPlaneCut::ExecuteFracture(const FFractureContext& FractureCont
 			if (LocalCutSettings->ReferenceActor != nullptr)
 			{
 				FTransform Transform(LocalCutSettings->ReferenceActor->GetActorTransform());
-				CuttingPlanes.Add(FPlane(Transform.GetLocation(), Transform.GetUnitAxis(EAxis::Z)));
+				CuttingPlanes.Add(FPlane(Transform.GetLocation() - FractureContext.Transform.GetLocation(), Transform.GetUnitAxis(EAxis::Z)));
 			}
 			else
 			{
