Angular and radial steps in radial fracture have been given a uimax of 50. #rb none #jira UE-76179

diff --git a/Engine/Plugins/Experimental/ChaosEditor/Source/FractureEditor/Private/FractureToolRadial.h b/Engine/Plugins/Experimental/ChaosEditor/Source/FractureEditor/Private/FractureToolRadial.h
--- a/Engine/Plugins/Experimental/ChaosEditor/Source/FractureEditor/Private/FractureToolRadial.h
+++ b/Engine/Plugins/Experimental/ChaosEditor/Source/FractureEditor/Private/FractureToolRadial.h
@@ -42,11 +42,11 @@ public:
 	float Radius;
 
 	/** Number of angular steps */
-	UPROPERTY(EditAnywhere, Category = RadialVoronoi, meta = (DisplayName = "Angular Steps", UIMin = "1", ClampMin = "1"))
+	UPROPERTY(EditAnywhere, Category = RadialVoronoi, meta = (DisplayName = "Angular Steps", UIMin = "1", UIMax = "50", ClampMin = "1"))
 	int AngularSteps;
 
 	/** Number of radial steps */
-	UPROPERTY(EditAnywhere, Category = RadialVoronoi, meta = (DisplayName = "Radial Steps", UIMin = "1", ClampMin = "1"))
+	UPROPERTY(EditAnywhere, Category = RadialVoronoi, meta = (DisplayName = "Radial Steps", UIMin = "1", UIMax = "50", ClampMin = "1"))
 	int RadialSteps;
 
 	/** Angle offset at each radial step */
