- Create Geometry Collection Asset will now use all static mesh components from all selected Actors.  #rb none #jira UE-76589

diff --git a/Engine/Plugins/Experimental/ChaosEditor/Source/FractureEditor/Private/FractureEditorModeToolkit.cpp b/Engine/Plugins/Experimental/ChaosEditor/Source/FractureEditor/Private/FractureEditorModeToolkit.cpp
--- a/Engine/Plugins/Experimental/ChaosEditor/Source/FractureEditor/Private/FractureEditorModeToolkit.cpp
+++ b/Engine/Plugins/Experimental/ChaosEditor/Source/FractureEditor/Private/FractureEditorModeToolkit.cpp
@@ -1183,23 +1183,39 @@ FReply FFractureEditorModeToolkit::OnFractureClicked()
 
 		FScopedTransaction Transaction(LOCTEXT("FractureMesh", "Fracture Mesh"));
 
+		TArray<UGeometryCollectionComponent*> NewComponents;
+
 		for (FFractureContext& FractureContext : FractureContexts)
 		{
 			ExecuteFracture(FractureContext);
 			UGeometryCollectionComponent* GeometryCollectionComponent = Cast<UGeometryCollectionComponent>(FractureContext.OriginalPrimitiveComponent);
-			FScopedColorEdit EditBoneColor = GeometryCollectionComponent->EditBoneSelection();
-			EditBoneColor.ResetBoneSelection();
-		
+			NewComponents.AddUnique(GeometryCollectionComponent);
+		}
+
+
+		for (UGeometryCollectionComponent* GeometryCollectionComponent : NewComponents)
+		{
+			FGeometryCollectionEdit GCEdit = GeometryCollectionComponent->EditRestCollection();
+			UGeometryCollection* GCObject = GCEdit.GetRestCollection();
+			TSharedPtr<FGeometryCollection, ESPMode::ThreadSafe> GeometryCollectionPtr = GCObject->GetGeometryCollection();
+			FGeometryCollectionClusteringUtility::UpdateHierarchyLevelOfChildren(GeometryCollectionPtr.Get(), -1);
+
+			FScopedColorEdit EditBoneColor(GeometryCollectionComponent, true);
+
+			EditBoneColor.SelectBones(GeometryCollection::ESelectionMode::None);
+			SetBoneSelection(GeometryCollectionComponent, EditBoneColor.GetSelectedBones(), true);
+
 			UpdateExplodedVectors(GeometryCollectionComponent);
-			SetOutlinerComponents({ GeometryCollectionComponent });
 
 			GeometryCollectionComponent->MarkRenderDynamicDataDirty();
 			GeometryCollectionComponent->MarkRenderStateDirty();
-
-			GCurrentLevelEditingViewportClient->Invalidate();
 		}
 
+		SetOutlinerComponents(NewComponents);
+
 		float ProcessingTime = static_cast<float>(FPlatformTime::Seconds() - CacheStartTime);
+
+		GCurrentLevelEditingViewportClient->Invalidate();
 	}
 
 	return FReply::Handled();
@@ -1503,30 +1519,11 @@ void FFractureEditorModeToolkit::GenerateAsset()
 
 	SelectionSet->GetSelectedObjects(SelectedActors);
 
-	for (AActor* Actor : SelectedActors)
-	{
-		TArray<UStaticMeshComponent*> Components;
-		Actor->GetComponents<UStaticMeshComponent>(Components);
-		if (Components.Num() > 0)
-		{
-			FFractureContext FractureContext;
-			FractureContext.OriginalActor = Actor;
-			FractureContext.OriginalPrimitiveComponent = Components[0];
-			FractureContext.Transform = Components[0]->GetComponentTransform();
-
-			AGeometryCollectionActor* GeometryCollectionActor = Cast<AGeometryCollectionActor>(Actor);
-			UGeometryCollectionComponent* GeometryCollectionComponent = Cast<UGeometryCollectionComponent>(FractureContext.OriginalPrimitiveComponent);
-
-			if (GeometryCollectionComponent == nullptr) // Create new GC actor from static mesh
-			{
-				OpenGenerateAssetDialog(FractureContext);
-			}
-		}
-	}
+	OpenGenerateAssetDialog(SelectedActors);
 }
 
 
-void FFractureEditorModeToolkit::OpenGenerateAssetDialog(FFractureContext InFractureContext)
+void FFractureEditorModeToolkit::OpenGenerateAssetDialog(TArray<AActor*>& Actors)
 {
 	TSharedPtr<SWindow> PickAssetPathWindow;
 
@@ -1543,7 +1540,7 @@ void FFractureEditorModeToolkit::OpenGenerateAssetDialog(FFractureContext InFrac
 		.AssetFilenameSuffix(TEXT("GeometryCollection"))
 		.HeadingText(LOCTEXT("CreateGeometryCollection_Heading", "Geometry Collection Name"))
 		.CreateButtonText(LOCTEXT("CreateGeometryCollection_ButtonLabel", "Create Geometry Collection"))
-		.OnCreateAssetAction(FOnPathChosen::CreateSP(this, &FFractureEditorModeToolkit::OnGenerateAssetPathChoosen, InFractureContext))
+		.OnCreateAssetAction(FOnPathChosen::CreateSP(this, &FFractureEditorModeToolkit::OnGenerateAssetPathChosen, Actors))
 	);
 
 	TSharedPtr<SWindow> RootWindow = FGlobalTabmanager::Get()->GetRootWindow();
@@ -1558,31 +1555,41 @@ void FFractureEditorModeToolkit::OpenGenerateAssetDialog(FFractureContext InFrac
 
 }
 
-void FFractureEditorModeToolkit::OnGenerateAssetPathChoosen(const FString& InAssetPath, FFractureContext InFractureContext)
+void FFractureEditorModeToolkit::OnGenerateAssetPathChosen(const FString& InAssetPath, TArray<AActor*> Actors)
 {
 		UGeometryCollectionComponent* GeometryCollectionComponent = nullptr;//  = Cast<UGeometryCollectionComponent>(FractureContext.OriginalPrimitiveComponent);
-		AGeometryCollectionActor* GeometryCollectionActor = Cast<AGeometryCollectionActor>(InFractureContext.OriginalActor);
-		GeometryCollectionActor = ConvertStaticMeshToGeometryCollection(InAssetPath, InFractureContext);
 
-		GeometryCollectionComponent = GeometryCollectionActor->GetGeometryCollectionComponent();
+		if (Actors.Num() > 0)
+		{
+			AActor* FirstActor = Actors[0];
 
-		FScopedColorEdit EditBoneColor = GeometryCollectionComponent->EditBoneSelection();
-		EditBoneColor.SetShowBoneColors(true);
+			AGeometryCollectionActor* GeometryCollectionActor = Cast<AGeometryCollectionActor>(FirstActor);
+			GeometryCollectionActor = ConvertStaticMeshToGeometryCollection(InAssetPath, Actors);
 
-		// Move GC actor to source actors position and remove source actor from scene
-		const FVector ActorLocation(InFractureContext.OriginalActor->GetActorLocation());
-		GeometryCollectionActor->SetActorLocation(ActorLocation);
-		InFractureContext.OriginalActor->Destroy();
+			GeometryCollectionComponent = GeometryCollectionActor->GetGeometryCollectionComponent();
 
-		GEditor->SelectActor(GeometryCollectionActor, true, true);
+			FScopedColorEdit EditBoneColor = GeometryCollectionComponent->EditBoneSelection();
+			EditBoneColor.SetShowBoneColors(true);
 
-		EditBoneColor.SelectBones(GeometryCollection::ESelectionMode::AllGeometry);
+			// Move GC actor to source actors position and remove source actor from scene
+			const FVector ActorLocation(FirstActor->GetActorLocation());
+			GeometryCollectionActor->SetActorLocation(ActorLocation);
 
-		SetOutlinerComponents({ GeometryCollectionComponent });
-		SetBoneSelection(GeometryCollectionComponent, EditBoneColor.GetSelectedBones(), true);
+			GEditor->SelectActor(GeometryCollectionActor, true, true);
 
-		GeometryCollectionComponent->MarkRenderDynamicDataDirty();
-		GeometryCollectionComponent->MarkRenderStateDirty();
+			EditBoneColor.SelectBones(GeometryCollection::ESelectionMode::AllGeometry);
+
+			SetOutlinerComponents({ GeometryCollectionComponent });
+			SetBoneSelection(GeometryCollectionComponent, EditBoneColor.GetSelectedBones(), true);
+
+			GeometryCollectionComponent->MarkRenderDynamicDataDirty();
+			GeometryCollectionComponent->MarkRenderStateDirty();
+
+			for (AActor* Actor : Actors)
+			{
+				Actor->Destroy();
+			}
+		}
 }
 
 
@@ -1846,36 +1853,44 @@ void FFractureEditorModeToolkit::UpdateExplodedVectors(UGeometryCollectionCompon
 	}
 }
 
-AGeometryCollectionActor*  FFractureEditorModeToolkit::ConvertStaticMeshToGeometryCollection(const FString& InAssetPath, FFractureContext& FractureContext)
+AGeometryCollectionActor*  FFractureEditorModeToolkit::ConvertStaticMeshToGeometryCollection(const FString& InAssetPath, TArray<AActor*>& Actors)
 {
-	const FString& Name = FractureContext.OriginalActor->GetActorLabel();
-	AGeometryCollectionActor* NewActor = CreateNewGeometryActor(InAssetPath, FTransform(), true);
+	ensure(Actors.Num() > 0);
+	AActor* FirstActor = Actors[0];
+ 	const FString& Name = FirstActor->GetActorLabel();
+	const FVector FirstActorLocation(FirstActor->GetActorLocation());
 
-	FGeometryCollectionEdit GeometryCollectionEdit = NewActor->GetGeometryCollectionComponent()->EditRestCollection();
-	FractureContext.FracturedGeometryCollection = GeometryCollectionEdit.GetRestCollection();
 
-	const FTransform ActorTransform(FractureContext.OriginalActor->GetTransform());
-
-	check(FractureContext.FracturedGeometryCollection);
-
-	TArray<UStaticMeshComponent*> StaticMeshComponents;
-	FractureContext.OriginalActor->GetComponents<UStaticMeshComponent>(StaticMeshComponents);
-	for (int32 ii = 0, ni = StaticMeshComponents.Num(); ii < ni; ++ii)
-	{
-		// We're partial to static mesh components, here
-		UStaticMeshComponent* StaticMeshComponent = StaticMeshComponents[ii];
-		if (StaticMeshComponent != nullptr)
-		{
-			UStaticMesh* ComponentStaticMesh = StaticMeshComponent->GetStaticMesh();
-			FTransform ComponentTranform(StaticMeshComponent->GetComponentTransform());
-			ComponentTranform.SetTranslation(ComponentTranform.GetTranslation() - ActorTransform.GetTranslation());
-			FGeometryCollectionConversion::AppendStaticMesh(ComponentStaticMesh, StaticMeshComponent, ComponentTranform, FractureContext.FracturedGeometryCollection, true);
-		}
-	}
-
-	FractureContext.FracturedGeometryCollection->InitializeMaterials();
+	AGeometryCollectionActor* NewActor = CreateNewGeometryActor(InAssetPath, FTransform(), true);
 
-	AddSingleRootNodeIfRequired(FractureContext.FracturedGeometryCollection);
+	FGeometryCollectionEdit GeometryCollectionEdit = NewActor->GetGeometryCollectionComponent()->EditRestCollection();
+	UGeometryCollection* FracturedGeometryCollection = GeometryCollectionEdit.GetRestCollection();
+
+ 	for (AActor* Actor : Actors)
+ 	{
+ 		const FTransform ActorTransform(Actor->GetTransform());
+		const FVector ActorOffset(Actor->GetActorLocation() - FirstActor->GetActorLocation());
+
+ 		check(FracturedGeometryCollection);
+ 
+ 		TArray<UStaticMeshComponent*> StaticMeshComponents;
+ 		Actor->GetComponents<UStaticMeshComponent>(StaticMeshComponents);
+ 		for (int32 ii = 0, ni = StaticMeshComponents.Num(); ii < ni; ++ii)
+ 		{
+ 			// We're partial to static mesh components, here
+ 			UStaticMeshComponent* StaticMeshComponent = StaticMeshComponents[ii];
+ 			if (StaticMeshComponent != nullptr)
+ 			{
+ 				UStaticMesh* ComponentStaticMesh = StaticMeshComponent->GetStaticMesh();
+ 				FTransform ComponentTranform(StaticMeshComponent->GetComponentTransform());
+ 				ComponentTranform.SetTranslation((ComponentTranform.GetTranslation() - ActorTransform.GetTranslation()) + ActorOffset);
+ 				FGeometryCollectionConversion::AppendStaticMesh(ComponentStaticMesh, StaticMeshComponent, ComponentTranform, FracturedGeometryCollection, true);
+ 			}
+ 		}
+ 
+ 		FracturedGeometryCollection->InitializeMaterials();
+ 	}
+ 	AddSingleRootNodeIfRequired(FracturedGeometryCollection);
 
 	return NewActor;
 }
diff --git a/Engine/Plugins/Experimental/ChaosEditor/Source/FractureEditor/Public/FractureEditorModeToolkit.h b/Engine/Plugins/Experimental/ChaosEditor/Source/FractureEditor/Public/FractureEditorModeToolkit.h
--- a/Engine/Plugins/Experimental/ChaosEditor/Source/FractureEditor/Public/FractureEditorModeToolkit.h
+++ b/Engine/Plugins/Experimental/ChaosEditor/Source/FractureEditor/Public/FractureEditorModeToolkit.h
@@ -79,15 +79,15 @@ public:
  	static AActor* AddActor(ULevel* InLevel, UClass* Class);
 	class AGeometryCollectionActor* CreateNewGeometryActor(const FString& Name, const FTransform& Transform, bool AddMaterials /*= false*/);
 
- 	void OpenGenerateAssetDialog(FFractureContext InFractureContext);
- 	void OnGenerateAssetPathChoosen(const FString& InAssetPath, FFractureContext InFractureContext);
+ 	void OpenGenerateAssetDialog(TArray<AActor*>& Actors);
+ 	void OnGenerateAssetPathChosen(const FString& InAssetPath, TArray<AActor*> Actors);
 
 	static void GetFractureContexts(TArray<FFractureContext>& FractureContexts);
 	void ExecuteFracture(FFractureContext& FractureContext);
 
 	static void GetSelectedGeometryCollectionComponents(TSet<UGeometryCollectionComponent*>& GeomCompSelection);
 
-	AGeometryCollectionActor* ConvertStaticMeshToGeometryCollection(const FString& InAssetPath, FFractureContext& FractureContext);
+	AGeometryCollectionActor* ConvertStaticMeshToGeometryCollection(const FString& InAssetPath, TArray<AActor*>& Actors);
 
 	static void AddSingleRootNodeIfRequired(UGeometryCollection* GeometryCollectionObject);
 	static void AddAdditionalAttributesIfRequired(UGeometryCollection* GeometryCollectionObject);
