Getting LevelCount from a GC now checks to see if the Levels array is empty range based loops don't seem to.  #rb none #jira UE-76010

diff --git a/Engine/Plugins/Experimental/ChaosEditor/Source/FractureEditor/Private/FractureEditorModeToolkit.cpp b/Engine/Plugins/Experimental/ChaosEditor/Source/FractureEditor/Private/FractureEditorModeToolkit.cpp
--- a/Engine/Plugins/Experimental/ChaosEditor/Source/FractureEditor/Private/FractureEditorModeToolkit.cpp
+++ b/Engine/Plugins/Experimental/ChaosEditor/Source/FractureEditor/Private/FractureEditorModeToolkit.cpp
@@ -757,11 +757,15 @@ int32 FFractureEditorModeToolkit::GetLevelCount()
 				if (HasLevelAttribute)
 				{
 					TManagedArray<int32>& Levels = GeometryCollection->GetAttribute<int32>("Level", FTransformCollection::TransformGroup);
-					for (int32 Level : Levels)
+
+					if(Levels.Num() > 0)
 					{
-						if (Level > ReturnLevel)
+						for (int32 Level : Levels)
 						{
-							ReturnLevel = Level;
+							if (Level > ReturnLevel)
+							{
+								ReturnLevel = Level;
+							}
 						}
 					}
 				}
