Fracture Editor: Fracture Plugin Mode unavailable unless bCompileChaos is true.

diff --git a/Engine/Plugins/Experimental/ChaosEditor/Source/FractureEditor/Private/FractureEditor.cpp b/Engine/Plugins/Experimental/ChaosEditor/Source/FractureEditor/Private/FractureEditor.cpp
--- a/Engine/Plugins/Experimental/ChaosEditor/Source/FractureEditor/Private/FractureEditor.cpp
+++ b/Engine/Plugins/Experimental/ChaosEditor/Source/FractureEditor/Private/FractureEditor.cpp
@@ -10,6 +10,9 @@
 
 void FFractureEditorModule::StartupModule()
 {
+
+#if INCLUDE_CHAOS
+
 	FFractureEditorStyle::Get();
 
 	// This code will execute after your module is loaded into memory; the exact timing is specified in the .uplugin file per-module
@@ -24,6 +27,9 @@ void FFractureEditorModule::StartupModule()
 	ToolBarExtensibilityManager = MakeShareable(new FExtensibilityManager);
 
 	FFractureEditorCommands::Register();
+
+#endif 
+	
 }
 
 void FFractureEditorModule::ShutdownModule()
