Fracture Editor:  Remove Command Access to Texture type Fracture

diff --git a/Engine/Plugins/Experimental/ChaosEditor/Source/FractureEditor/Private/FractureToolBitmap.cpp b/Engine/Plugins/Experimental/ChaosEditor/Source/FractureEditor/Private/FractureToolBitmap.cpp
--- a/Engine/Plugins/Experimental/ChaosEditor/Source/FractureEditor/Private/FractureToolBitmap.cpp
+++ b/Engine/Plugins/Experimental/ChaosEditor/Source/FractureEditor/Private/FractureToolBitmap.cpp
@@ -29,8 +29,10 @@ FSlateIcon UFractureToolBitmap::GetToolIcon() const
 
 void UFractureToolBitmap::RegisterUICommand( FFractureEditorCommands* BindingContext ) 
 {
+	/*
 	UI_COMMAND_EXT( BindingContext, UICommandInfo, "Texture", "Texture", "Texture Voronoi Fracture", EUserInterfaceActionType::ToggleButton, FInputChord() );
 	BindingContext->Texture = UICommandInfo;
+	*/
 }
 
 TArray<UObject*> UFractureToolBitmap::GetSettingsObjects() const 
