Fracture Editor: Fixes Panel sizing issues by adding a resizable splitter

diff --git a/Engine/Plugins/Experimental/ChaosEditor/Source/FractureEditor/Private/FractureEditorModeToolkit.cpp b/Engine/Plugins/Experimental/ChaosEditor/Source/FractureEditor/Private/FractureEditorModeToolkit.cpp
--- a/Engine/Plugins/Experimental/ChaosEditor/Source/FractureEditor/Private/FractureEditorModeToolkit.cpp
+++ b/Engine/Plugins/Experimental/ChaosEditor/Source/FractureEditor/Private/FractureEditorModeToolkit.cpp
@@ -15,6 +15,7 @@
 #include "Widgets/Layout/SBox.h"
 #include "Widgets/Layout/SExpandableArea.h"
 #include "Widgets/Layout/SScrollBox.h"
+#include "Widgets/Layout/SSplitter.h"
 #include "Widgets/Input/SSlider.h"
 #include "Widgets/Input/SCheckBox.h"
 #include "Widgets/Input/SComboBox.h"
@@ -82,8 +83,8 @@ FFractureEditorModeToolkit::FFractureEditorModeToolkit()
 	, AutoClusterMode(EFractureAutoClusterMode::BoundingBox)
 	, AutoClusterSiteCount(10)
 	, ActiveTool(nullptr)
-	, InnerFractureSlot(nullptr)
-	, InnerOutlinerSlot(nullptr)
+	, bFractureGroupExpanded(true)
+	, bOutlinerGroupExpanded(true)
 {
 }
 
@@ -393,74 +394,79 @@ void FFractureEditorModeToolkit::Init(const TSharedPtr<IToolkitHost>& InitToolki
 		]
 
 		+SVerticalBox::Slot()
-		.Expose(InnerFractureSlot)
 		[
-			SNew(SExpandableArea)
-			.AreaTitle(FText(LOCTEXT("Fracture", "Fracture")))
-			.HeaderPadding(FMargin(2.0, 2.0))
-			.Padding(FMargin(MorePadding))
-			.BorderImage(FEditorStyle::Get().GetBrush("DetailsView.CategoryTop"))
-			.BorderBackgroundColor( FLinearColor( .6,.6,.6, 1.0f ) )
-			.BodyBorderBackgroundColor (FLinearColor( 1.0, 0.0, 0.0))
-			.AreaTitleFont(FEditorStyle::Get().GetFontStyle("DetailsView.CategoryFontStyle"))
-			.OnAreaExpansionChanged( FOnBooleanValueChanged::CreateLambda( [=](bool NewIsExpanded) { NewIsExpanded ? InnerFractureSlot->FillHeight(1.0f) : InnerFractureSlot->AutoHeight(); }) )
-			.BodyContent()
-			[
-				SNew(SScrollBox)
+			SNew(SSplitter)
+			.Orientation( Orient_Vertical )
 
-				+SScrollBox::Slot()
-				.HAlign(HAlign_Center)
+			+SSplitter::Slot()
+			.SizeRule(TAttribute<SSplitter::ESizeRule>(this, &FFractureEditorModeToolkit::GetFractureGroupSizeRule ))
+			[
+				SNew(SExpandableArea)
+				.AreaTitle(FText(LOCTEXT("Fracture", "Fracture")))
+				.HeaderPadding(FMargin(2.0, 2.0))
+				.Padding(FMargin(MorePadding))
+				.BorderImage(FEditorStyle::Get().GetBrush("DetailsView.CategoryTop"))
+				.BorderBackgroundColor( FLinearColor( .6,.6,.6, 1.0f ) )
+				.BodyBorderBackgroundColor (FLinearColor( 1.0, 0.0, 0.0))
+				.AreaTitleFont(FEditorStyle::Get().GetFontStyle("DetailsView.CategoryFontStyle"))
+				.OnAreaExpansionChanged(this, &FFractureEditorModeToolkit::OnFractureGroupExpansionChanged)
+				.BodyContent()
 				[
-					FractureToolBuilder.MakeWidget()
-				]
+					SNew(SScrollBox)
 
-				+SScrollBox::Slot()
-				.Padding(0.0f, MorePadding)
-				[
-					DetailsView.ToSharedRef()
-				]
-	
-				+SScrollBox::Slot()
-				[
-					SNew(SHorizontalBox)
-					+SHorizontalBox::Slot()
-					.FillWidth(1)
+					+SScrollBox::Slot()
+					.HAlign(HAlign_Center)
+					[
+						FractureToolBuilder.MakeWidget()
+					]
 
-					+SHorizontalBox::Slot()
-					.AutoWidth()
+					+SScrollBox::Slot()
+					.Padding(0.0f, MorePadding)
 					[
-						SNew( SButton )
-						.HAlign(HAlign_Center)
-						.ContentPadding(FMargin(MorePadding, Padding))
-						.OnClicked(this, &FFractureEditorModeToolkit::OnFractureClicked)
-						.IsEnabled( this, &FFractureEditorModeToolkit::CanExecuteFracture)
-						.Text(FText(LOCTEXT("FractureFractureButton", "Fracture")))
+						DetailsView.ToSharedRef()
 					]
+		
+					+SScrollBox::Slot()
+					[
+						SNew(SHorizontalBox)
+						+SHorizontalBox::Slot()
+						.FillWidth(1)
 
-					+SHorizontalBox::Slot()
-					.FillWidth(1)
-				]	
+						+SHorizontalBox::Slot()
+						.AutoWidth()
+						[
+							SNew( SButton )
+							.HAlign(HAlign_Center)
+							.ContentPadding(FMargin(MorePadding, Padding))
+							.OnClicked(this, &FFractureEditorModeToolkit::OnFractureClicked)
+							.IsEnabled( this, &FFractureEditorModeToolkit::CanExecuteFracture)
+							.Text(FText(LOCTEXT("FractureFractureButton", "Fracture")))
+						]
+
+						+SHorizontalBox::Slot()
+						.FillWidth(1)
+					]	
+				]
 			]
-		]
 
 
-		+SVerticalBox::Slot()
-		.Expose(InnerOutlinerSlot)
-		.FillHeight(.5f)
-		[
-			SNew(SExpandableArea)
-			.AreaTitle(FText(LOCTEXT("Outliner", "Outliner")))
-			.HeaderPadding(FMargin(2.0, 2.0))
-			.Padding(FMargin(MorePadding))
-			.BorderImage(FEditorStyle::Get().GetBrush("DetailsView.CategoryTop"))
-			.BorderBackgroundColor( FLinearColor( .6,.6,.6, 1.0f ) )
-			.BodyBorderBackgroundColor (FLinearColor( 1.0, 0.0, 0.0))
-			.AreaTitleFont(FEditorStyle::Get().GetFontStyle("DetailsView.CategoryFontStyle"))
-			.OnAreaExpansionChanged( FOnBooleanValueChanged::CreateLambda( [=](bool NewIsExpanded) { NewIsExpanded ? InnerOutlinerSlot->FillHeight(0.5f) : InnerOutlinerSlot->AutoHeight(); }) )
-			.BodyContent()
+			+SSplitter::Slot()
+			.SizeRule(TAttribute<SSplitter::ESizeRule>(this, &FFractureEditorModeToolkit::GetOutlinerGroupSizeRule ))
 			[
-				SAssignNew(OutlinerView, SGeometryCollectionOutliner )
-				.OnBoneSelectionChanged(this, &FFractureEditorModeToolkit::OnOutlinerBoneSelectionChanged)
+				SNew(SExpandableArea)
+				.AreaTitle(FText(LOCTEXT("Outliner", "Outliner")))
+				.HeaderPadding(FMargin(2.0, 2.0))
+				.Padding(FMargin(MorePadding))
+				.BorderImage(FEditorStyle::Get().GetBrush("DetailsView.CategoryTop"))
+				.BorderBackgroundColor( FLinearColor( .6,.6,.6, 1.0f ) )
+				.BodyBorderBackgroundColor (FLinearColor( 1.0, 0.0, 0.0))
+				.AreaTitleFont(FEditorStyle::Get().GetFontStyle("DetailsView.CategoryFontStyle"))
+				.OnAreaExpansionChanged(this, &FFractureEditorModeToolkit::OnOutlinerGroupExpansionChanged)
+				.BodyContent()
+				[
+					SAssignNew(OutlinerView, SGeometryCollectionOutliner )
+					.OnBoneSelectionChanged(this, &FFractureEditorModeToolkit::OnOutlinerBoneSelectionChanged)
+				]
 			]
 		]
 	];
diff --git a/Engine/Plugins/Experimental/ChaosEditor/Source/FractureEditor/Public/FractureEditorModeToolkit.h b/Engine/Plugins/Experimental/ChaosEditor/Source/FractureEditor/Public/FractureEditorModeToolkit.h
--- a/Engine/Plugins/Experimental/ChaosEditor/Source/FractureEditor/Public/FractureEditorModeToolkit.h
+++ b/Engine/Plugins/Experimental/ChaosEditor/Source/FractureEditor/Public/FractureEditorModeToolkit.h
@@ -6,6 +6,7 @@
 #include "Widgets/DeclarativeSyntaxSupport.h"
 #include "Toolkits/BaseToolkit.h"
 #include "Widgets/Input/SCheckBox.h"
+#include "Widgets/Layout/SSplitter.h"
 #include "GeometryCollection/GeometryCollectionActor.h"
 #include "AutoClusterFracture.h"
 
@@ -102,6 +103,15 @@ public:
 	void SetAutoClusterSiteCount(uint32 InSiteCount);
 	uint32 GetAutoClusterSiteCount() const;
 
+	/** Fracture Group Layout Management */
+	SSplitter::ESizeRule GetFractureGroupSizeRule() const { return bFractureGroupExpanded ? SSplitter::ESizeRule::FractionOfParent : SSplitter::ESizeRule::SizeToContent; }
+	void OnFractureGroupExpansionChanged(bool bExpanded) { bFractureGroupExpanded = bExpanded; }
+
+	/** Outliner Group Layout Management */
+	SSplitter::ESizeRule GetOutlinerGroupSizeRule() const { return bOutlinerGroupExpanded ? SSplitter::ESizeRule::FractionOfParent : SSplitter::ESizeRule::SizeToContent; }
+	void OnOutlinerGroupExpansionChanged(bool bExpanded) { bOutlinerGroupExpanded = bExpanded; }
+
+
 protected:
 	static bool IsGeometryCollectionSelected();
 	static bool IsStaticMeshSelected();
@@ -118,9 +128,8 @@ private:
 
 	UFractureTool* ActiveTool;
 
-	// TSharedPtr<SVerticalBox::Slot> FractureAreaSlot;
-	SVerticalBox::FSlot* InnerFractureSlot = nullptr;
-	SVerticalBox::FSlot* InnerOutlinerSlot = nullptr;
+	bool bFractureGroupExpanded = true;
+	bool bOutlinerGroupExpanded = true;
 	
 	TSharedPtr<IDetailsView> DetailsView;
 	TSharedPtr<SWidget> ToolkitWidget;
