Fracture Editor: Fixes Exploded View getting stuck on 1%

diff --git a/Engine/Plugins/Experimental/ChaosEditor/Source/FractureEditor/Private/FractureEditorModeToolkit.cpp b/Engine/Plugins/Experimental/ChaosEditor/Source/FractureEditor/Private/FractureEditorModeToolkit.cpp
--- a/Engine/Plugins/Experimental/ChaosEditor/Source/FractureEditor/Private/FractureEditorModeToolkit.cpp
+++ b/Engine/Plugins/Experimental/ChaosEditor/Source/FractureEditor/Private/FractureEditorModeToolkit.cpp
@@ -712,7 +712,7 @@ void FFractureEditorModeToolkit::OnSetShowBoneColors()
 
 void FFractureEditorModeToolkit::OnSetExplodedViewValue(float NewValue)
 {
-	if (!FMath::IsNearlyEqual(ExplodeAmount,NewValue, .01f))
+	if ( FMath::Abs<float>( ExplodeAmount - NewValue ) >= .01f)
 	{
 		ExplodeAmount = NewValue;
 
@@ -741,6 +741,8 @@ void FFractureEditorModeToolkit::OnSetExplodedViewValue(float NewValue)
 				}
 			}
 		}
+		
+		GCurrentLevelEditingViewportClient->Invalidate();
 	}
 }
 
@@ -808,6 +810,8 @@ void FFractureEditorModeToolkit::OnSetLevelViewValue(int32 NewValue)
 		}
 	}
 	SetOutlinerComponents(GeomCompSelection.Array());
+
+	GCurrentLevelEditingViewportClient->Invalidate();
 }
 
 void FFractureEditorModeToolkit::ViewUpOneLevel()
