Check for StaticMesh in actor in now recursive.
Conversion of BP to GC now collects static mesh components from children in actor as well.

diff --git a/Engine/Plugins/Experimental/ChaosEditor/Source/FractureEditor/Private/FractureEditorModeToolkit.cpp b/Engine/Plugins/Experimental/ChaosEditor/Source/FractureEditor/Private/FractureEditorModeToolkit.cpp
--- a/Engine/Plugins/Experimental/ChaosEditor/Source/FractureEditor/Private/FractureEditorModeToolkit.cpp
+++ b/Engine/Plugins/Experimental/ChaosEditor/Source/FractureEditor/Private/FractureEditorModeToolkit.cpp
@@ -1381,7 +1381,7 @@ void FFractureEditorModeToolkit::OnCluster()
 					FGeometryCollectionClusteringUtility::UpdateHierarchyLevelOfChildren(GeometryCollection, -1);
 
 					FScopedColorEdit EditBoneColor = GeometryCollectionComponent->EditBoneSelection();
- 					EditBoneColor.ResetBoneSelection();
+					EditBoneColor.ResetBoneSelection();
 					EditBoneColor.ResetHighlightedBones();
 					GeometryCollectionComponent->MarkRenderDynamicDataDirty();
 					GeometryCollectionComponent->MarkRenderStateDirty();
@@ -1409,8 +1409,8 @@ void FFractureEditorModeToolkit::OnUncluster()
 			TSharedPtr<FGeometryCollection, ESPMode::ThreadSafe> GeometryCollectionPtr = GeometryCollectionObject->GetGeometryCollection();
 			if (FGeometryCollection* GeometryCollection = GeometryCollectionPtr.Get())
 			{
-				FGeometryCollectionClusteringUtility::CollapseSelectedHierarchy(FractureLevel, GeometryCollectionComponent->GetSelectedBones(), GeometryCollection);
 				FGeometryCollectionClusteringUtility::UpdateHierarchyLevelOfChildren(GeometryCollection, -1);
+				FGeometryCollectionClusteringUtility::CollapseSelectedHierarchy(FractureLevel, GeometryCollectionComponent->GetSelectedBones(), GeometryCollection);
 
 				FScopedColorEdit EditBoneColor = GeometryCollectionComponent->EditBoneSelection();
 				EditBoneColor.ResetBoneSelection();
@@ -1721,14 +1721,15 @@ bool FFractureEditorModeToolkit::IsGeometryCollectionSelected()
 bool FFractureEditorModeToolkit::IsStaticMeshSelected()
 {
 	USelection* SelectedActors = GEditor->GetSelectedActors();
-	TArray<ULevel*> UniqueLevels;
 	for (FSelectionIterator Iter(*SelectedActors); Iter; ++Iter)
 	{
 		AActor* Actor = Cast<AActor>(*Iter);
 		if (Actor)
 		{
-			TArray<UActorComponent*> Components = Actor->GetComponentsByClass(UStaticMeshComponent::StaticClass());
-			if (Components.Num() > 0)
+			TInlineComponentArray<UStaticMeshComponent*> StaticMeshComponents;
+			Actor->GetComponents<UStaticMeshComponent>(StaticMeshComponents, true);
+
+			if (StaticMeshComponents.Num() > 0)
 			{
 				return true;
 			}
@@ -1778,7 +1779,6 @@ bool GetValidGeoCenter(const TManagedArray<int32>& TransformToGeometryIndex, con
 	return false;
 }
 
-
 void FFractureEditorModeToolkit::UpdateExplodedVectors(UGeometryCollectionComponent* GeometryCollectionComponent) const
 {
 	FGeometryCollectionEdit RestCollection = GeometryCollectionComponent->EditRestCollection(GeometryCollection::EEditUpdate::RestPhysicsDynamic);
@@ -1871,7 +1871,7 @@ AGeometryCollectionActor*  FFractureEditorModeToolkit::ConvertStaticMeshToGeomet
 {
 	ensure(Actors.Num() > 0);
 	AActor* FirstActor = Actors[0];
- 	const FString& Name = FirstActor->GetActorLabel();
+	const FString& Name = FirstActor->GetActorLabel();
 	const FVector FirstActorLocation(FirstActor->GetActorLocation());
 
 
@@ -1880,31 +1880,31 @@ AGeometryCollectionActor*  FFractureEditorModeToolkit::ConvertStaticMeshToGeomet
 	FGeometryCollectionEdit GeometryCollectionEdit = NewActor->GetGeometryCollectionComponent()->EditRestCollection(GeometryCollection::EEditUpdate::RestPhysicsDynamic);
 	UGeometryCollection* FracturedGeometryCollection = GeometryCollectionEdit.GetRestCollection();
 
- 	for (AActor* Actor : Actors)
- 	{
- 		const FTransform ActorTransform(Actor->GetTransform());
+	for (AActor* Actor : Actors)
+	{
+		const FTransform ActorTransform(Actor->GetTransform());
 		const FVector ActorOffset(Actor->GetActorLocation() - FirstActor->GetActorLocation());
 
- 		check(FracturedGeometryCollection);
+		check(FracturedGeometryCollection);
  
- 		TArray<UStaticMeshComponent*> StaticMeshComponents;
- 		Actor->GetComponents<UStaticMeshComponent>(StaticMeshComponents);
- 		for (int32 ii = 0, ni = StaticMeshComponents.Num(); ii < ni; ++ii)
- 		{
- 			// We're partial to static mesh components, here
- 			UStaticMeshComponent* StaticMeshComponent = StaticMeshComponents[ii];
- 			if (StaticMeshComponent != nullptr)
- 			{
- 				UStaticMesh* ComponentStaticMesh = StaticMeshComponent->GetStaticMesh();
- 				FTransform ComponentTranform(StaticMeshComponent->GetComponentTransform());
- 				ComponentTranform.SetTranslation((ComponentTranform.GetTranslation() - ActorTransform.GetTranslation()) + ActorOffset);
- 				FGeometryCollectionConversion::AppendStaticMesh(ComponentStaticMesh, StaticMeshComponent, ComponentTranform, FracturedGeometryCollection, true);
- 			}
- 		}
+		TArray<UStaticMeshComponent*> StaticMeshComponents;
+		Actor->GetComponents<UStaticMeshComponent>(StaticMeshComponents, true);
+		for (int32 ii = 0, ni = StaticMeshComponents.Num(); ii < ni; ++ii)
+		{
+			// We're partial to static mesh components, here
+			UStaticMeshComponent* StaticMeshComponent = StaticMeshComponents[ii];
+			if (StaticMeshComponent != nullptr)
+			{
+				UStaticMesh* ComponentStaticMesh = StaticMeshComponent->GetStaticMesh();
+				FTransform ComponentTranform(StaticMeshComponent->GetComponentTransform());
+				ComponentTranform.SetTranslation((ComponentTranform.GetTranslation() - ActorTransform.GetTranslation()) + ActorOffset);
+				FGeometryCollectionConversion::AppendStaticMesh(ComponentStaticMesh, StaticMeshComponent, ComponentTranform, FracturedGeometryCollection, true);
+			}
+		}
  
- 		FracturedGeometryCollection->InitializeMaterials();
- 	}
- 	AddSingleRootNodeIfRequired(FracturedGeometryCollection);
+		FracturedGeometryCollection->InitializeMaterials();
+	}
+	AddSingleRootNodeIfRequired(FracturedGeometryCollection);
 
 	return NewActor;
 }
