Tooltip typo fixes in cluster facture.  #jira UE-76665  #rb trivial

diff --git a/Engine/Plugins/Experimental/ChaosEditor/Source/FractureEditor/Private/FractureToolCluster.h b/Engine/Plugins/Experimental/ChaosEditor/Source/FractureEditor/Private/FractureToolCluster.h
--- a/Engine/Plugins/Experimental/ChaosEditor/Source/FractureEditor/Private/FractureToolCluster.h
+++ b/Engine/Plugins/Experimental/ChaosEditor/Source/FractureEditor/Private/FractureToolCluster.h
@@ -37,22 +37,22 @@ public:
 	UPROPERTY(EditAnywhere, Category = ClusterVoronoi, meta = (DisplayName = "Maximum Cluster Sites", UIMin = "1", UIMax = "2000", ClampMin = "1"))
 	int32 NumberClustersMax;
 
-	/** Sites per of Clusters - Cluster Voronoi Method */
+	/** Sites per # of Clusters - Cluster Voronoi Method */
 	UPROPERTY(EditAnywhere, Category = ClusterVoronoi, meta = (DisplayName = "Minimum Sites Per Cluster", UIMin = "0", UIMax = "2000", ClampMin = "0"))
 	int32 SitesPerClusterMin;
 
 	UPROPERTY(EditAnywhere, Category = ClusterVoronoi, meta = (DisplayName = "Maximum Sites Per Cluster", UIMin = "0", UIMax = "200", ClampMin = "0"))
 	int32 SitesPerClusterMax;
 
-	/** Clusters Radius - Cluster Voronoi Method */
+	/** Cluster's Radius - Cluster Voronoi Method */
 	UPROPERTY(EditAnywhere, Category = ClusterVoronoi, meta = (DisplayName = "Minimum distance from center as part of bounds max extent", UIMin = "0.0", UIMax = "1.0"))
 	float ClusterRadiusPercentageMin;
 
-	/** Clusters Radius - Cluster Voronoi Method */
+	/** Cluster's Radius - Cluster Voronoi Method */
 	UPROPERTY(EditAnywhere, Category = ClusterVoronoi, meta = (DisplayName = "Maximum distance from center as part of bounds max extent", UIMin = "0.0", UIMax = "1.0"))
 	float ClusterRadiusPercentageMax;
 
-	/** Clusters Radius - Cluster Voronoi Method */
+	/** Cluster's Radius - Cluster Voronoi Method */
 	UPROPERTY(EditAnywhere, Category = ClusterVoronoi)
 	float ClusterRadius;
 
