all max values in cluster fracture set to 200.  #jira UE-76935 #rb trivial

diff --git a/Engine/Plugins/Experimental/ChaosEditor/Source/FractureEditor/Private/FractureToolCluster.h b/Engine/Plugins/Experimental/ChaosEditor/Source/FractureEditor/Private/FractureToolCluster.h
--- a/Engine/Plugins/Experimental/ChaosEditor/Source/FractureEditor/Private/FractureToolCluster.h
+++ b/Engine/Plugins/Experimental/ChaosEditor/Source/FractureEditor/Private/FractureToolCluster.h
@@ -14,7 +14,7 @@ class UFractureClusterSettings
 	GENERATED_BODY()
 public:
 
-	UFractureClusterSettings() 
+	UFractureClusterSettings()
 	: NumberClustersMin(8)
 	, NumberClustersMax(8)
 	, SitesPerClusterMin(2)
@@ -30,15 +30,15 @@ public:
 #endif
 
 	/** Number of Clusters - Cluster Voronoi Method */
-	UPROPERTY(EditAnywhere, Category = ClusterVoronoi, meta = (DisplayName = "Minimum Cluster Sites", UIMin = "1", UIMax = "2000", ClampMin = "1"))
+	UPROPERTY(EditAnywhere, Category = ClusterVoronoi, meta = (DisplayName = "Minimum Cluster Sites", UIMin = "1", UIMax = "200", ClampMin = "1"))
 	int32 NumberClustersMin;
 
 	/** Number of Clusters - Cluster Voronoi Method */
-	UPROPERTY(EditAnywhere, Category = ClusterVoronoi, meta = (DisplayName = "Maximum Cluster Sites", UIMin = "1", UIMax = "2000", ClampMin = "1"))
+	UPROPERTY(EditAnywhere, Category = ClusterVoronoi, meta = (DisplayName = "Maximum Cluster Sites", UIMin = "1", UIMax = "200", ClampMin = "1"))
 	int32 NumberClustersMax;
 
 	/** Sites per # of Clusters - Cluster Voronoi Method */
-	UPROPERTY(EditAnywhere, Category = ClusterVoronoi, meta = (DisplayName = "Minimum Sites Per Cluster", UIMin = "0", UIMax = "2000", ClampMin = "0"))
+	UPROPERTY(EditAnywhere, Category = ClusterVoronoi, meta = (DisplayName = "Minimum Sites Per Cluster", UIMin = "0", UIMax = "200", ClampMin = "0"))
 	int32 SitesPerClusterMin;
 
 	UPROPERTY(EditAnywhere, Category = ClusterVoronoi, meta = (DisplayName = "Maximum Sites Per Cluster", UIMin = "0", UIMax = "200", ClampMin = "0"))
@@ -75,7 +75,7 @@ public:
 	virtual TArray<UObject*> GetSettingsObjects() const override;// { return TArray<UObject*>(); }
 
 	virtual void RegisterUICommand( FFractureEditorCommands* BindingContext );
-	
+
 	// virtual void ExecuteFracture() {}
 	// virtual bool CanExecuteFracture() { return true; }
 
