Cluster Francture now gets its own settings inside it's site generation callback.  #jira UE-76975 #rb trivial

diff --git a/Engine/Plugins/Experimental/ChaosEditor/Source/FractureEditor/Private/FractureToolCluster.cpp b/Engine/Plugins/Experimental/ChaosEditor/Source/FractureEditor/Private/FractureToolCluster.cpp
--- a/Engine/Plugins/Experimental/ChaosEditor/Source/FractureEditor/Private/FractureToolCluster.cpp
+++ b/Engine/Plugins/Experimental/ChaosEditor/Source/FractureEditor/Private/FractureToolCluster.cpp
@@ -65,29 +65,30 @@ TArray<UObject*> UFractureToolCluster::GetSettingsObjects() const
 void UFractureToolCluster::GenerateVoronoiSites(const FFractureContext &Context, TArray<FVector>& Sites)
 {
  	const UFractureCommonSettings* LocalCommonSettings = GetDefault<UFractureCommonSettings>();
+	const UFractureClusterSettings* ClusterSettings = GetDefault<UFractureClusterSettings>();
 
 	FRandomStream RandStream(Context.RandomSeed);
-	const int32 SiteCount = RandStream.RandRange(Settings->NumberClustersMin, Settings->NumberClustersMax);
+	const int32 SiteCount = RandStream.RandRange(ClusterSettings->NumberClustersMin, ClusterSettings->NumberClustersMax);
 
 	const FVector Extent(Context.Bounds.Max - Context.Bounds.Min);
 
 	TArray<FVector> CenterSites;
 
 	Sites.Reserve(Sites.Num() + SiteCount);
-	for (int32 ii = 0; ii < Settings->NumberClustersMax; ++ii)
+	for (int32 ii = 0; ii < ClusterSettings->NumberClustersMax; ++ii)
 	{
 		CenterSites.Emplace(Context.Bounds.Min + FVector(RandStream.FRand(), RandStream.FRand(), RandStream.FRand()) * Extent);
 	}
 
 	for (int32 kk = 0, nk = CenterSites.Num(); kk < nk; ++kk)
 	{
-		const int32 SubSiteCount = RandStream.RandRange(Settings->SitesPerClusterMin, Settings->SitesPerClusterMax);
+		const int32 SubSiteCount = RandStream.RandRange(ClusterSettings->SitesPerClusterMin, ClusterSettings->SitesPerClusterMax);
 
 		for (int32 ii = 0; ii < SubSiteCount; ++ii)
 		{
 			FVector V(RandStream.VRand());
 			V.Normalize();
-			V *= Settings->ClusterRadius + (RandStream.FRandRange(Settings->ClusterRadiusPercentageMin, Settings->ClusterRadiusPercentageMax) * Context.Bounds.GetExtent().GetAbsMax());
+			V *= ClusterSettings->ClusterRadius + (RandStream.FRandRange(ClusterSettings->ClusterRadiusPercentageMin, ClusterSettings->ClusterRadiusPercentageMax) * Context.Bounds.GetExtent().GetAbsMax());
 			V += CenterSites[kk];
 			Sites.Emplace(V);
 		}
