Fixed variable name in site loop.  #jira UE-76931 #rb trivial

diff --git a/Engine/Plugins/Experimental/ChaosEditor/Source/FractureEditor/Private/FractureToolCluster.cpp b/Engine/Plugins/Experimental/ChaosEditor/Source/FractureEditor/Private/FractureToolCluster.cpp
--- a/Engine/Plugins/Experimental/ChaosEditor/Source/FractureEditor/Private/FractureToolCluster.cpp
+++ b/Engine/Plugins/Experimental/ChaosEditor/Source/FractureEditor/Private/FractureToolCluster.cpp
@@ -27,37 +27,37 @@ void UFractureClusterSettings::PostEditChangeChainProperty(struct FPropertyChang
 
 
 
-UFractureToolCluster::UFractureToolCluster(const FObjectInitializer& ObjInit) 
-	: Super(ObjInit) 
+UFractureToolCluster::UFractureToolCluster(const FObjectInitializer& ObjInit)
+	: Super(ObjInit)
 {
 	Settings = NewObject<UFractureClusterSettings>(GetTransientPackage(), UFractureClusterSettings::StaticClass());
 	Settings->OwnerTool = this;
 }
 
 FText UFractureToolCluster::GetDisplayText() const
-{ 
-	return FText(NSLOCTEXT("Fracture", "FractureToolCluster", "Cluster Voronoi")); 
+{
+	return FText(NSLOCTEXT("Fracture", "FractureToolCluster", "Cluster Voronoi"));
 }
 
-FText UFractureToolCluster::GetTooltipText() const 
-{ 
-	return FText(NSLOCTEXT("Fracture", "FractureToolClusterTooltip", "Cluster Voronoi Fracture")); 
+FText UFractureToolCluster::GetTooltipText() const
+{
+	return FText(NSLOCTEXT("Fracture", "FractureToolClusterTooltip", "Cluster Voronoi Fracture"));
 }
 
-FSlateIcon UFractureToolCluster::GetToolIcon() const 
+FSlateIcon UFractureToolCluster::GetToolIcon() const
 {
 	return FSlateIcon("FractureEditorStyle", "FractureEditor.Clustered");
 }
 
-void UFractureToolCluster::RegisterUICommand( FFractureEditorCommands* BindingContext ) 
+void UFractureToolCluster::RegisterUICommand( FFractureEditorCommands* BindingContext )
 {
 	UI_COMMAND_EXT( BindingContext, UICommandInfo, "Clustered", "Clustered", "Clustered Voronoi Fracture", EUserInterfaceActionType::ToggleButton, FInputChord() );
 	BindingContext->Clustered = UICommandInfo;
 }
 
-TArray<UObject*> UFractureToolCluster::GetSettingsObjects() const 
-{ 
-	TArray<UObject*> ReturnSettings; 
+TArray<UObject*> UFractureToolCluster::GetSettingsObjects() const
+{
+	TArray<UObject*> ReturnSettings;
 	ReturnSettings.Add(GetMutableDefault<UFractureClusterSettings>());
 	return ReturnSettings;
 }
@@ -75,7 +75,7 @@ void UFractureToolCluster::GenerateVoronoiSites(const FFractureContext &Context,
 	TArray<FVector> CenterSites;
 
 	Sites.Reserve(Sites.Num() + SiteCount);
-	for (int32 ii = 0; ii < ClusterSettings->NumberClustersMax; ++ii)
+	for (int32 ii = 0; ii < SiteCount; ++ii)
 	{
 		CenterSites.Emplace(Context.Bounds.Min + FVector(RandStream.FRand(), RandStream.FRand(), RandStream.FRand()) * Extent);
 	}
