Invalidate level viewport on toggling ShowBoneColors in fracture editor #jira UE-77171 #rb trivial

diff --git a/Engine/Plugins/Experimental/ChaosEditor/Source/FractureEditor/Private/FractureEditorModeToolkit.cpp b/Engine/Plugins/Experimental/ChaosEditor/Source/FractureEditor/Private/FractureEditorModeToolkit.cpp
--- a/Engine/Plugins/Experimental/ChaosEditor/Source/FractureEditor/Private/FractureEditorModeToolkit.cpp
+++ b/Engine/Plugins/Experimental/ChaosEditor/Source/FractureEditor/Private/FractureEditorModeToolkit.cpp
@@ -706,7 +706,7 @@ void FFractureEditorModeToolkit::OnSetShowBoneColors()
 		Comp->MarkRenderStateDirty();
 		Comp->MarkRenderDynamicDataDirty();
 	}
-
+	GCurrentLevelEditingViewportClient->Invalidate();
 }
 
 
