Fracture cluster now sorts the selected bones list so processing happens in order compatible with ClusterBonesUnderNewNode()  #jira UE-76587 #rb none

diff --git a/Engine/Plugins/Experimental/ChaosEditor/Source/FractureEditor/Private/FractureEditorModeToolkit.cpp b/Engine/Plugins/Experimental/ChaosEditor/Source/FractureEditor/Private/FractureEditorModeToolkit.cpp
--- a/Engine/Plugins/Experimental/ChaosEditor/Source/FractureEditor/Private/FractureEditorModeToolkit.cpp
+++ b/Engine/Plugins/Experimental/ChaosEditor/Source/FractureEditor/Private/FractureEditorModeToolkit.cpp
@@ -1375,9 +1375,26 @@ void FFractureEditorModeToolkit::OnCluster()
 				TSharedPtr<FGeometryCollection, ESPMode::ThreadSafe> GeometryCollectionPtr = GCObject->GetGeometryCollection();
 				if (FGeometryCollection* GeometryCollection = GeometryCollectionPtr.Get())
 				{
+					const TManagedArray<TSet<int32>>& Children = GeometryCollection->GetAttribute<TSet<int32>>("Children", FGeometryCollection::TransformGroup);
+
+					// sort the selection list so ClusterBonesUnderNewNode() happens in the correct order for leaf nodes
+					TArray<int32> SortedSelectedBones;
+					SortedSelectedBones.Reserve(SelectedBones.Num());
+					for (int32 SelectedBone : SelectedBones)
+					{
+						if(Children[SelectedBone].Num() > 0)
+						{
+							SortedSelectedBones.Insert(SelectedBone, 0);
+						}
+						else
+						{
+							SortedSelectedBones.Add(SelectedBone);
+						}
+					}
 					// cluster Selected Bones under the first selected bone
-					int32 InsertAtIndex = SelectedBones[0];
-					FGeometryCollectionClusteringUtility::ClusterBonesUnderNewNode(GeometryCollection, InsertAtIndex, SelectedBones, false);
+					int32 InsertAtIndex = SortedSelectedBones[0];
+
+					FGeometryCollectionClusteringUtility::ClusterBonesUnderNewNode(GeometryCollection, InsertAtIndex, SortedSelectedBones, false);
 					FGeometryCollectionClusteringUtility::UpdateHierarchyLevelOfChildren(GeometryCollection, -1);
 
 					FScopedColorEdit EditBoneColor = GeometryCollectionComponent->EditBoneSelection();
