clear old bricks before creating new ones

#rnx

diff --git a/Engine/Plugins/Experimental/ChaosEditor/Source/FractureEditor/Private/FractureToolBrick.cpp b/Engine/Plugins/Experimental/ChaosEditor/Source/FractureEditor/Private/FractureToolBrick.cpp
--- a/Engine/Plugins/Experimental/ChaosEditor/Source/FractureEditor/Private/FractureToolBrick.cpp
+++ b/Engine/Plugins/Experimental/ChaosEditor/Source/FractureEditor/Private/FractureToolBrick.cpp
@@ -242,7 +242,7 @@ void UFractureToolBrick::PostEditChangeChainProperty(struct FPropertyChangedChai
 
 	SelectionSet->GetSelectedObjects(SelectedActors);
 	BrickTransforms.Empty();
-	Edges.Empty();;
+	Edges.Empty();
 
 	for (AActor* Actor : SelectedActors)
 	{
@@ -329,6 +329,7 @@ void UFractureToolBrick::ExecuteFracture(const FFractureContext& FractureContext
  		TSharedPtr<FGeometryCollection, ESPMode::ThreadSafe> GeometryCollectionPtr = FractureContext.FracturedGeometryCollection->GetGeometryCollection();
  		if (FGeometryCollection* GeometryCollection = GeometryCollectionPtr.Get())
  		{
+			BrickTransforms.Empty();
 			GenerateBrickTransforms(FractureContext.Bounds);
 
 			const UFractureCommonSettings* LocalCommonSettings = GetDefault<UFractureCommonSettings>();
diff --git a/Engine/Plugins/Experimental/PlanarCutPlugin/Source/PlanarCut/Public/PlanarCut.h b/Engine/Plugins/Experimental/PlanarCutPlugin/Source/PlanarCut/Public/PlanarCut.h
--- a/Engine/Plugins/Experimental/PlanarCutPlugin/Source/PlanarCut/Public/PlanarCut.h
+++ b/Engine/Plugins/Experimental/PlanarCutPlugin/Source/PlanarCut/Public/PlanarCut.h
@@ -60,7 +60,7 @@ struct PLANARCUT_API FPlanarCells
 	 * Debugging function to check that the plane boundary vertices are wound to match the orientation of the plane normal vectors
 	 * @return	false if any plane boundary vertices are found in the 'wrong' orientation relative to the plane normal
 	 */
-	bool HasValidPlaneBoundaryOrientations()
+	bool HasValidPlaneBoundaryOrientations() const
 	{
 		for (int32 PlaneIdx = 0; PlaneIdx < PlaneBoundaries.Num(); PlaneIdx++)
 		{
