Fracture Editor: Fixes Crash related to undo/redo not clearning cached selection state

diff --git a/Engine/Plugins/Experimental/ChaosEditor/Source/FractureEditor/Private/SGeometryCollectionOutliner.cpp b/Engine/Plugins/Experimental/ChaosEditor/Source/FractureEditor/Private/SGeometryCollectionOutliner.cpp
--- a/Engine/Plugins/Experimental/ChaosEditor/Source/FractureEditor/Private/SGeometryCollectionOutliner.cpp
+++ b/Engine/Plugins/Experimental/ChaosEditor/Source/FractureEditor/Private/SGeometryCollectionOutliner.cpp
@@ -108,15 +108,10 @@ void SGeometryCollectionOutliner::SetComponents(const TArray<UGeometryCollection
 	for (UGeometryCollectionComponent* Component : InNewComponents)
 	{
 		RootNodes.Add(MakeShared<FGeometryCollectionTreeItemComponent>(Component));
-
 		TArray<int32> SelectedBones = Component->GetSelectedBones();
-		if (SelectedBones.Num())
-		{
-			SetBoneSelection(Component, SelectedBones, false);
-		}
+		SetBoneSelection(Component, SelectedBones, !SelectedBones.Num());
 	}
 
-
 	TreeView->RequestTreeRefresh();
 	ExpandAll();
 }
