Reduce code paths that rebuild GeometryCollection levelset in Fracture Editor.

diff --git a/Engine/Plugins/Experimental/ChaosEditor/Source/FractureEditor/Private/FractureEditorMode.cpp b/Engine/Plugins/Experimental/ChaosEditor/Source/FractureEditor/Private/FractureEditorMode.cpp
--- a/Engine/Plugins/Experimental/ChaosEditor/Source/FractureEditor/Private/FractureEditorMode.cpp
+++ b/Engine/Plugins/Experimental/ChaosEditor/Source/FractureEditor/Private/FractureEditorMode.cpp
@@ -343,7 +343,7 @@ void FFractureEditorMode::GetActorGlobalBounds(TArrayView<UGeometryCollectionCom
 {
 	for (UGeometryCollectionComponent* GeometryCollectionComponent : GeometryComponents)
 	{
-		FGeometryCollectionEdit RestCollection = GeometryCollectionComponent->EditRestCollection();
+		FGeometryCollectionEdit RestCollection = GeometryCollectionComponent->EditRestCollection(GeometryCollection::EEditUpdate::None);
 		UGeometryCollection* GeometryCollection = RestCollection.GetRestCollection();
 
 		TSharedPtr<FGeometryCollection, ESPMode::ThreadSafe> GeometryCollectionPtr = GeometryCollection->GetGeometryCollection();
diff --git a/Engine/Plugins/Experimental/ChaosEditor/Source/FractureEditor/Private/FractureEditorModeToolkit.cpp b/Engine/Plugins/Experimental/ChaosEditor/Source/FractureEditor/Private/FractureEditorModeToolkit.cpp
--- a/Engine/Plugins/Experimental/ChaosEditor/Source/FractureEditor/Private/FractureEditorModeToolkit.cpp
+++ b/Engine/Plugins/Experimental/ChaosEditor/Source/FractureEditor/Private/FractureEditorModeToolkit.cpp
@@ -987,7 +987,7 @@ void FFractureEditorModeToolkit::SetOutlinerComponents(const TArray<UGeometryCol
 {
 	for (UGeometryCollectionComponent* Component : InNewComponents)
 	{
-		FGeometryCollectionEdit RestCollection = Component->EditRestCollection();
+		FGeometryCollectionEdit RestCollection = Component->EditRestCollection(GeometryCollection::EEditUpdate::None);
 		UGeometryCollection* FracturedGeometryCollection = RestCollection.GetRestCollection();
 		TSharedPtr<FGeometryCollection, ESPMode::ThreadSafe> GeometryCollectionPtr = FracturedGeometryCollection->GetGeometryCollection();
 
@@ -1080,11 +1080,10 @@ void FFractureEditorModeToolkit::GetFractureContexts(TArray<FFractureContext>& F
 				UPrimitiveComponent* PrimitiveComponent = CastChecked<UPrimitiveComponent>(Component);
 				if (UGeometryCollectionComponent* GeometryCollectionComponent = Cast<UGeometryCollectionComponent>(PrimitiveComponent))
 				{
-					FGeometryCollectionEdit RestCollection = GeometryCollectionComponent->EditRestCollection();
+					FGeometryCollectionEdit RestCollection = GeometryCollectionComponent->EditRestCollection(GeometryCollection::EEditUpdate::None);
 					UGeometryCollection* FracturedGeometryCollection = RestCollection.GetRestCollection();
 
-					FScopedColorEdit EditBoneColor = GeometryCollectionComponent->EditBoneSelection();
-					const TArray<int32>& SelectedBonesOriginal = EditBoneColor.GetSelectedBones();
+					const TArray<int32>& SelectedBonesOriginal = GeometryCollectionComponent->GetSelectedBones();
 
 					TSharedPtr<FGeometryCollection, ESPMode::ThreadSafe> GeometryCollectionPtr = FracturedGeometryCollection->GetGeometryCollection();
 					FGeometryCollection* OutGeometryCollection = GeometryCollectionPtr.Get();
@@ -1266,11 +1265,10 @@ bool FFractureEditorModeToolkit::IsLeafBoneSelected() const
 
 		if (SelectedBones.Num() > 0)
 		{
-			FGeometryCollectionEdit GCEdit = GeometryCollectionComponent->EditRestCollection();
-			if (UGeometryCollection* GCObject = GCEdit.GetRestCollection())
+			if (const UGeometryCollection* GCObject = GeometryCollectionComponent->GetRestCollection())
 			{
 				TSharedPtr<FGeometryCollection, ESPMode::ThreadSafe> GeometryCollectionPtr = GCObject->GetGeometryCollection();
-				if (FGeometryCollection* GeometryCollection = GeometryCollectionPtr.Get())
+				if (const FGeometryCollection* GeometryCollection = GeometryCollectionPtr.Get())
 				{
 					const TManagedArray<TSet<int32>>& Children = GeometryCollection->GetAttribute<TSet<int32>>("Children", FGeometryCollection::TransformGroup);
 
@@ -1798,18 +1796,24 @@ bool GetValidGeoCenter(const TManagedArray<int32>& TransformToGeometryIndex, con
 
 void FFractureEditorModeToolkit::UpdateExplodedVectors(UGeometryCollectionComponent* GeometryCollectionComponent) const
 {
-	FGeometryCollectionEdit RestCollection = GeometryCollectionComponent->EditRestCollection(GeometryCollection::EEditUpdate::RestPhysicsDynamic);
-	UGeometryCollection* GeometryCollection = RestCollection.GetRestCollection();
-
-	TSharedPtr<FGeometryCollection, ESPMode::ThreadSafe> GeometryCollectionPtr = GeometryCollection->GetGeometryCollection();
-	FGeometryCollection* OutGeometryCollection = GeometryCollectionPtr.Get();
+	TSharedPtr<FGeometryCollection, ESPMode::ThreadSafe> GeometryCollectionPtr = GeometryCollectionComponent->GetRestCollection()->GetGeometryCollection();
+	const FGeometryCollection* OutGeometryCollectionConst = GeometryCollectionPtr.Get();
 
 	if (FMath::IsNearlyEqual(ExplodeAmount, 0.0f))
 	{
-		OutGeometryCollection->RemoveAttribute("ExplodedVector", FGeometryCollection::TransformGroup);
+		if (OutGeometryCollectionConst->HasAttribute("ExplodedVector", FGeometryCollection::TransformGroup))
+		{
+			FGeometryCollectionEdit RestCollection = GeometryCollectionComponent->EditRestCollection(GeometryCollection::EEditUpdate::RestPhysicsDynamic);
+			FGeometryCollection* OutGeometryCollection = GeometryCollectionPtr.Get();
+			OutGeometryCollection->RemoveAttribute("ExplodedVector", FGeometryCollection::TransformGroup);
+		}
 	}
 	else
 	{
+		FGeometryCollectionEdit RestCollection = GeometryCollectionComponent->EditRestCollection(GeometryCollection::EEditUpdate::RestPhysicsDynamic);
+		UGeometryCollection* GeometryCollection = RestCollection.GetRestCollection();
+		FGeometryCollection* OutGeometryCollection = GeometryCollectionPtr.Get();
+
 		if (!OutGeometryCollection->HasAttribute("ExplodedVector", FGeometryCollection::TransformGroup))
 		{
 			OutGeometryCollection->AddAttribute<FVector>("ExplodedVector", FGeometryCollection::TransformGroup, FManagedArrayCollection::FConstructionParameters(FName(), false));
diff --git a/Engine/Source/Runtime/Experimental/GeometryCollectionEngine/Private/GeometryCollection/GeometryCollectionComponent.cpp b/Engine/Source/Runtime/Experimental/GeometryCollectionEngine/Private/GeometryCollection/GeometryCollectionComponent.cpp
--- a/Engine/Source/Runtime/Experimental/GeometryCollectionEngine/Private/GeometryCollection/GeometryCollectionComponent.cpp
+++ b/Engine/Source/Runtime/Experimental/GeometryCollectionEngine/Private/GeometryCollection/GeometryCollectionComponent.cpp
@@ -1596,21 +1596,18 @@ bool FScopedColorEdit::IsBoneSelected(int BoneIndex) const
 
 void FScopedColorEdit::SetSelectedBones(const TArray<int32>& SelectedBonesIn)
 {
-	Component->Modify();
 	bUpdated = true;
 	Component->SelectedBones = SelectedBonesIn;
 }
 
 void FScopedColorEdit::AppendSelectedBones(const TArray<int32>& SelectedBonesIn)
 {
-	Component->Modify();
 	bUpdated = true;
 	Component->SelectedBones.Append(SelectedBonesIn);
 }
 
 void FScopedColorEdit::ToggleSelectedBones(const TArray<int32>& SelectedBonesIn)
 {
-	Component->Modify();
 	bUpdated = true;
 	for (int32 BoneIndex : SelectedBonesIn)
 	{
@@ -1629,7 +1626,6 @@ void FScopedColorEdit::AddSelectedBone(int32 BoneIndex)
 {
 	if (!Component->SelectedBones.Contains(BoneIndex))
 	{
-		Component->Modify();
 		bUpdated = true;
 		Component->SelectedBones.Push(BoneIndex);
 	}
@@ -1639,7 +1635,6 @@ void FScopedColorEdit::ClearSelectedBone(int32 BoneIndex)
 {
 	if (Component->SelectedBones.Contains(BoneIndex))
 	{
-		Component->Modify();
 		bUpdated = true;
 		Component->SelectedBones.Remove(BoneIndex);
 	}
@@ -1654,7 +1649,6 @@ void FScopedColorEdit::ResetBoneSelection()
 {
 	if (Component->SelectedBones.Num() > 0)
 	{
-		Component->Modify();
 		bUpdated = true;
 	}
 
@@ -1665,7 +1659,6 @@ void FScopedColorEdit::SelectBones(GeometryCollection::ESelectionMode SelectionM
 {
 	check(Component);
 
-	Component->Modify();
 	const UGeometryCollection* GeometryCollection = Component->GetRestCollection();
 	if (GeometryCollection)
 	{
@@ -1857,7 +1850,7 @@ void FScopedColorEdit::UpdateBoneColors()
 {
 	// @todo FractureTools - For large fractures updating colors this way is extremely slow because the render state (and thus all buffers) must be recreated.
 	// It would be better to push the update to the proxy via a render command and update the existing buffer directly
-	FGeometryCollectionEdit GeometryCollectionEdit = Component->EditRestCollection();
+	FGeometryCollectionEdit GeometryCollectionEdit = Component->EditRestCollection(GeometryCollection::EEditUpdate::None);
 	UGeometryCollection* GeometryCollection = GeometryCollectionEdit.GetRestCollection();
 	if(GeometryCollection)
 	{
