Modify GC component on bone selection change. #jira UE-76567 #rb trivial

diff --git a/Engine/Plugins/Experimental/ChaosEditor/Source/FractureEditor/Private/FractureEditorMode.cpp b/Engine/Plugins/Experimental/ChaosEditor/Source/FractureEditor/Private/FractureEditorMode.cpp
--- a/Engine/Plugins/Experimental/ChaosEditor/Source/FractureEditor/Private/FractureEditorMode.cpp
+++ b/Engine/Plugins/Experimental/ChaosEditor/Source/FractureEditor/Private/FractureEditorMode.cpp
@@ -154,6 +154,7 @@ bool FFractureEditorMode::HandleClick(FEditorViewportClient* InViewportClient, H
 
 			FScopedTransaction Transaction(FractureTransactionContexts::SelectBoneContext, LOCTEXT("SelectGeometryCollectionBoneTransaction", "Select Bone"), GeometryCollectionProxy->Component);
 			FFractureSelectionTools::ToggleSelectedBones(GeometryCollectionProxy->Component, BoneIndices, !Click.IsControlDown());
+			GeometryCollectionProxy->Component->Modify();
 
 			if (Toolkit.IsValid())
 			{
