Disable Chaos Plugins by default, to ensure they aren't used without setting bCompileChaos.

diff --git a/Engine/Plugins/Experimental/ChaosEditor/ChaosEditor.uplugin b/Engine/Plugins/Experimental/ChaosEditor/ChaosEditor.uplugin
--- a/Engine/Plugins/Experimental/ChaosEditor/ChaosEditor.uplugin
+++ b/Engine/Plugins/Experimental/ChaosEditor/ChaosEditor.uplugin
@@ -13,7 +13,7 @@
 	"CanContainContent": false,
 	"IsBetaVersion": true,
 	"Installed": false,
-	"EnabledByDefault": true,
+	"EnabledByDefault": false,
 	"Modules": [
 		{
 			"Name": "FractureEditor",
