Added a bunch of tooltips to fracture tools.  #jira none #rb none #fyi jack.oakman

diff --git a/Engine/Plugins/Experimental/ChaosEditor/Source/FractureEditor/Private/FractureEditorCommands.cpp b/Engine/Plugins/Experimental/ChaosEditor/Source/FractureEditor/Private/FractureEditorCommands.cpp
--- a/Engine/Plugins/Experimental/ChaosEditor/Source/FractureEditor/Private/FractureEditorCommands.cpp
+++ b/Engine/Plugins/Experimental/ChaosEditor/Source/FractureEditor/Private/FractureEditorCommands.cpp
@@ -8,21 +8,21 @@
 
 #define LOCTEXT_NAMESPACE "FractureEditorCommands"
 
-FFractureEditorCommands::FFractureEditorCommands() 
+FFractureEditorCommands::FFractureEditorCommands()
 	: TCommands<FFractureEditorCommands>("FractureEditor", LOCTEXT("Fracture", "Fracture"), NAME_None, FFractureEditorStyle::StyleName)
 {
 }
 
 void FFractureEditorCommands::RegisterCommands()
 {
-	
+
 	// Selection Commands
 	UI_COMMAND(SelectAll, "SelectAll", "Selects all Bones in the GeometryCollection", EUserInterfaceActionType::Button, FInputChord(EModifierKey::Control | EModifierKey::Shift, EKeys::A));
 	UI_COMMAND(SelectNone, "SelectNone", "Deselects all Bones in the GeometryCollection", EUserInterfaceActionType::Button, FInputChord(EModifierKey::Control | EModifierKey::Shift, EKeys::D));
-	UI_COMMAND(SelectNeighbors, "SelectNeighbors", "Select Neighbors", EUserInterfaceActionType::Button, FInputChord() );
-	UI_COMMAND(SelectSiblings, "SelectSiblings", "Select Siblings", EUserInterfaceActionType::Button, FInputChord() );
-	UI_COMMAND(SelectAllInCluster, "SelectAllInCluster", "Select All In Cluster", EUserInterfaceActionType::Button, FInputChord() );
-	UI_COMMAND(SelectInvert, "SelectInvert", "Invert The Selection", EUserInterfaceActionType::Button, FInputChord() );
+	UI_COMMAND(SelectNeighbors, "SelectNeighbors", "Select all bones adjacent to the currently selected bones", EUserInterfaceActionType::Button, FInputChord() );
+	UI_COMMAND(SelectSiblings, "SelectSiblings", "Select all bones at the same levels as the currently selected bones", EUserInterfaceActionType::Button, FInputChord() );
+	UI_COMMAND(SelectAllInCluster, "SelectAllInCluster", "Select all bones with the same parent as selected bones", EUserInterfaceActionType::Button, FInputChord() );
+	UI_COMMAND(SelectInvert, "SelectInvert", "Invert the selected bones", EUserInterfaceActionType::Button, FInputChord() );
 
 	// View settings
 	UI_COMMAND(ToggleShowBoneColors, "ShowBoneColors", "Toggle Show Bone Colors", EUserInterfaceActionType::ToggleButton, FInputChord(EModifierKey::Shift, EKeys::B) );
@@ -32,15 +32,15 @@ void FFractureEditorCommands::RegisterCommands()
 	UI_COMMAND(ExplodeLess, "ExplodeLess", "Explode 10% Less", EUserInterfaceActionType::Button, FInputChord(EModifierKey::Shift, EKeys::Q) );
 
 	// Cluster Operations
-	UI_COMMAND(Flatten, "Flatten", "Flatten", EUserInterfaceActionType::Button, FInputChord() );
+	UI_COMMAND(Flatten, "Flatten", "Flattens all bones to level 1", EUserInterfaceActionType::Button, FInputChord() );
 	UI_COMMAND(FlattenToLevel, "FlattenToLevel", "Flatten Only up to the Current View Level", EUserInterfaceActionType::Button, FInputChord() );
 	UI_COMMAND(AutoCluster, "AutoCluster", "Auto Cluster", EUserInterfaceActionType::Button, FInputChord() );
-	UI_COMMAND(Cluster, "Cluster", "Cluster", EUserInterfaceActionType::Button, FInputChord() );
+	UI_COMMAND(Cluster, "Cluster", "Clusters selected bones under a new parent", EUserInterfaceActionType::Button, FInputChord() );
 	UI_COMMAND(Uncluster, "Uncluster", "Uncluster", EUserInterfaceActionType::Button, FInputChord() );
 	UI_COMMAND(Merge, "Merge", "Merge", EUserInterfaceActionType::Button, FInputChord() );
 	UI_COMMAND(MoveUp, "MoveUp", "Move Bones Up One Level", EUserInterfaceActionType::Button, FInputChord() );
 
-	UI_COMMAND(GenerateAsset, "GenerateAsset", "Generate Geometry Collection Asset", EUserInterfaceActionType::Button, FInputChord() );
+	UI_COMMAND(GenerateAsset, "GenerateAsset", "Generate Geometry Collection Asset from static meshes contained in selected actors", EUserInterfaceActionType::Button, FInputChord() );
 
 
 	for (TObjectIterator<UClass> ClassIterator; ClassIterator; ++ClassIterator)
diff --git a/Engine/Plugins/Experimental/ChaosEditor/Source/FractureEditor/Private/FractureEditorModeToolkit.cpp b/Engine/Plugins/Experimental/ChaosEditor/Source/FractureEditor/Private/FractureEditorModeToolkit.cpp
--- a/Engine/Plugins/Experimental/ChaosEditor/Source/FractureEditor/Private/FractureEditorModeToolkit.cpp
+++ b/Engine/Plugins/Experimental/ChaosEditor/Source/FractureEditor/Private/FractureEditorModeToolkit.cpp
@@ -77,7 +77,7 @@ TArray<UClass*> FindFractureToolClasses()
 	return Classes;
 }
 
-FFractureEditorModeToolkit::FFractureEditorModeToolkit() 
+FFractureEditorModeToolkit::FFractureEditorModeToolkit()
 	: ExplodeAmount(0.0f)
 	, FractureLevel(-1)
 	, AutoClusterMode(EFractureAutoClusterMode::BoundingBox)
@@ -129,14 +129,14 @@ void FFractureEditorModeToolkit::Init(const TSharedPtr<IToolkitHost>& InitToolki
 	FPropertyEditorModule& EditModule = FModuleManager::Get().GetModuleChecked<FPropertyEditorModule>("PropertyEditor");
 
 	FDetailsViewArgs DetailsViewArgs(
-		false,  //bUpdateFromSelection= 
-		false, //bLockable= 
-		false, //bAllowSearch= 
+		false,  //bUpdateFromSelection=
+		false, //bLockable=
+		false, //bAllowSearch=
 		FDetailsViewArgs::HideNameArea,
-		true, //bHideSelectionTip= 
-		nullptr, //InNotifyHook= 
-		false, //InSearchInitialKeyFocus= 
-		NAME_None); //InViewIdentifier= 
+		true, //bHideSelectionTip=
+		nullptr, //InNotifyHook=
+		false, //InSearchInitialKeyFocus=
+		NAME_None); //InViewIdentifier=
 	DetailsViewArgs.DefaultsOnlyVisibility = EEditDefaultsOnlyNodeVisibility::Automatic;
 	DetailsViewArgs.bShowOptions = false;
 	DetailsViewArgs.bAllowMultipleTopLevelObjects = true;
@@ -157,7 +157,7 @@ void FFractureEditorModeToolkit::Init(const TSharedPtr<IToolkitHost>& InitToolki
 		SNew(SVerticalBox)
 
 		+SVerticalBox::Slot()
-		.Padding(0.0f, Padding) 
+		.Padding(0.0f, Padding)
 		.AutoHeight()
 		.HAlign(HAlign_Center)
 		[
@@ -185,7 +185,7 @@ void FFractureEditorModeToolkit::Init(const TSharedPtr<IToolkitHost>& InitToolki
 				+SVerticalBox::Slot()
 				.HAlign(HAlign_Center)
 				[
-					SelectionBuilder.MakeWidget()	
+					SelectionBuilder.MakeWidget()
 				]
 			]
 		]
@@ -229,6 +229,7 @@ void FFractureEditorModeToolkit::Init(const TSharedPtr<IToolkitHost>& InitToolki
 					.MaxValue(100)
 					.Delta(1)
 					.Font(FCoreStyle::GetDefaultFontStyle("Regular", 18))
+					.ToolTipText(LOCTEXT("FractureEditor.Exploded_Tooltip", "How much to seperate the drawing of the bones to aid in setup.  Does not effect simulation"))
 				]
 
 				+SHorizontalBox::Slot()
@@ -269,13 +270,13 @@ void FFractureEditorModeToolkit::Init(const TSharedPtr<IToolkitHost>& InitToolki
 				.AutoWidth()
 				[
 					SNew(SComboButton)
-					.ContentPadding(0)	
+					.ContentPadding(0)
 					.ButtonStyle(FEditorStyle::Get(), "ToolBar.Button")
 					.ForegroundColor(FEditorStyle::Get().GetSlateColor("ToolBar.SToolBarComboButtonBlock.ComboButton.Color"))
 					.ButtonContent()
 					[
 						SNew(STextBlock)
-						.Text_Lambda( [=]() -> FText { 
+						.Text_Lambda( [=]() -> FText {
 							if (FractureLevel < 0)
 							{
 								return LOCTEXT("FractureViewAllLevels", "All");
@@ -289,6 +290,8 @@ void FFractureEditorModeToolkit::Init(const TSharedPtr<IToolkitHost>& InitToolki
 
 						})
 						.Font(FCoreStyle::GetDefaultFontStyle("Regular", 18))
+						.ToolTipText(LOCTEXT("FractureEditor.Level_Tooltip", "Set the currently view level of the geometry collection"))
+
 					]
 					.OnGetMenuContent(this, &FFractureEditorModeToolkit::GetLevelViewMenuContent)
 				]
@@ -308,6 +311,7 @@ void FFractureEditorModeToolkit::Init(const TSharedPtr<IToolkitHost>& InitToolki
 					[
 						SNew(SImage)
 						.Image(FFractureEditorStyle::Get().GetBrush("FractureEditor.Visibility"))
+						.ToolTipText(LOCTEXT("FractureEditor.Visibility_Tooltip", "Toggle showing the bone colours of the geometry collection"))
 					]
 					.OnGetMenuContent(this, &FFractureEditorModeToolkit::GetViewMenuContent)
 				]
@@ -365,7 +369,7 @@ void FFractureEditorModeToolkit::Init(const TSharedPtr<IToolkitHost>& InitToolki
 							SNew(STextBlock)
 							.Text(this, &FFractureEditorModeToolkit::GetAutoClusterModeDisplayName)
 							.Font(FCoreStyle::GetDefaultFontStyle("Regular", 8))
-						]	
+						]
 					]
 
 					+SHorizontalBox::Slot()
@@ -388,6 +392,8 @@ void FFractureEditorModeToolkit::Init(const TSharedPtr<IToolkitHost>& InitToolki
 						.MaxValue(9999)
 						.Value(this, &FFractureEditorModeToolkit::GetAutoClusterSiteCount)
 						.OnValueChanged(this, &FFractureEditorModeToolkit::SetAutoClusterSiteCount)
+						.ToolTipText(LOCTEXT("AutoClusterSitesTooltip", "How many groups you would like to try to autocluster into"))
+
 					]
 				]
 			]
@@ -425,7 +431,7 @@ void FFractureEditorModeToolkit::Init(const TSharedPtr<IToolkitHost>& InitToolki
 					[
 						DetailsView.ToSharedRef()
 					]
-		
+
 					+SScrollBox::Slot()
 					[
 						SNew(SHorizontalBox)
@@ -445,7 +451,7 @@ void FFractureEditorModeToolkit::Init(const TSharedPtr<IToolkitHost>& InitToolki
 
 						+SHorizontalBox::Slot()
 						.FillWidth(1)
-					]	
+					]
 				]
 			]
 
@@ -470,7 +476,7 @@ void FFractureEditorModeToolkit::Init(const TSharedPtr<IToolkitHost>& InitToolki
 			]
 		]
 	];
-	
+
 
 	// Bind Chaos Commands;
 	BindCommands();
@@ -486,103 +492,103 @@ void FFractureEditorModeToolkit::BindCommands()
 	const FFractureEditorCommands& Commands = FFractureEditorCommands::Get();
 
 	ToolkitCommands->MapAction(
-		Commands.SelectAll, 
+		Commands.SelectAll,
 		FExecuteAction::CreateSP(
-			this, 
-			&FFractureEditorModeToolkit::OnSelectByMode, 
+			this,
+			&FFractureEditorModeToolkit::OnSelectByMode,
 			GeometryCollection::ESelectionMode::AllGeometry
 			),
 		FCanExecuteAction::CreateStatic(&FFractureEditorModeToolkit::IsGeometryCollectionSelected)
 		);
 
 	ToolkitCommands->MapAction(
-		Commands.SelectNone, 
+		Commands.SelectNone,
 		FExecuteAction::CreateSP(
-			this, 
-			&FFractureEditorModeToolkit::OnSelectByMode, 
+			this,
+			&FFractureEditorModeToolkit::OnSelectByMode,
 			GeometryCollection::ESelectionMode::None
 			),
 		FCanExecuteAction::CreateStatic(&FFractureEditorModeToolkit::IsGeometryCollectionSelected)
 		);
 
 	ToolkitCommands->MapAction(
-		Commands.SelectNeighbors, 
+		Commands.SelectNeighbors,
 		FExecuteAction::CreateSP(
-			this, 
-			&FFractureEditorModeToolkit::OnSelectByMode, 
+			this,
+			&FFractureEditorModeToolkit::OnSelectByMode,
 			GeometryCollection::ESelectionMode::Neighbors
 			),
 		FCanExecuteAction::CreateStatic(&FFractureEditorModeToolkit::IsGeometryCollectionSelected)
 		);
-	
+
 	ToolkitCommands->MapAction(
-		Commands.SelectSiblings, 
+		Commands.SelectSiblings,
 		FExecuteAction::CreateSP(
-			this, 
-			&FFractureEditorModeToolkit::OnSelectByMode, 
+			this,
+			&FFractureEditorModeToolkit::OnSelectByMode,
 			GeometryCollection::ESelectionMode::Siblings
 			),
 		FCanExecuteAction::CreateStatic(&FFractureEditorModeToolkit::IsGeometryCollectionSelected)
 		);
 
 	ToolkitCommands->MapAction(
-		Commands.SelectAllInCluster, 
+		Commands.SelectAllInCluster,
 		FExecuteAction::CreateSP(
-			this, 
-			&FFractureEditorModeToolkit::OnSelectByMode, 
+			this,
+			&FFractureEditorModeToolkit::OnSelectByMode,
 			GeometryCollection::ESelectionMode::AllInCluster
 			),
 		FCanExecuteAction::CreateStatic(&FFractureEditorModeToolkit::IsGeometryCollectionSelected)
 		);
 
 	ToolkitCommands->MapAction(
-		Commands.SelectInvert, 
+		Commands.SelectInvert,
 		FExecuteAction::CreateSP(
-			this, 
-			&FFractureEditorModeToolkit::OnSelectByMode, 
+			this,
+			&FFractureEditorModeToolkit::OnSelectByMode,
 			GeometryCollection::ESelectionMode::InverseGeometry
 			),
 		FCanExecuteAction::CreateStatic(&FFractureEditorModeToolkit::IsGeometryCollectionSelected)
 		);
 
 	ToolkitCommands->MapAction(
-		Commands.AutoCluster, 
+		Commands.AutoCluster,
 		FExecuteAction::CreateSP(this, &FFractureEditorModeToolkit::OnAutoCluster),
 		FCanExecuteAction::CreateStatic(&FFractureEditorModeToolkit::IsGeometryCollectionSelected)
 	);
 
 	ToolkitCommands->MapAction(
-		Commands.Cluster, 
+		Commands.Cluster,
 		FExecuteAction::CreateSP(this, &FFractureEditorModeToolkit::OnCluster),
 		FCanExecuteAction::CreateStatic(&FFractureEditorModeToolkit::IsGeometryCollectionSelected)
 	);
-	
+
 	ToolkitCommands->MapAction(
-		Commands.Uncluster, 
+		Commands.Uncluster,
 		FExecuteAction::CreateSP(this, &FFractureEditorModeToolkit::OnUncluster),
 		FCanExecuteAction::CreateStatic(&FFractureEditorModeToolkit::IsGeometryCollectionSelected)
 	);
 
 	ToolkitCommands->MapAction(
-		Commands.Flatten, 
+		Commands.Flatten,
 		FExecuteAction::CreateSP(this, &FFractureEditorModeToolkit::OnFlatten),
 		FCanExecuteAction::CreateStatic(&FFractureEditorModeToolkit::IsGeometryCollectionSelected)
 	);
 
 	ToolkitCommands->MapAction(
-		Commands.FlattenToLevel, 
+		Commands.FlattenToLevel,
 		FExecuteAction::CreateSP(this, &FFractureEditorModeToolkit::OnFlattenToLevel),
 		FCanExecuteAction::CreateLambda([] { return false; })
 		);
 
 // 	ToolkitCommands->MapAction(
-// 		Commands.Merge, 
+// 		Commands.Merge,
 // 		FExecuteAction::CreateSP(this, &FFractureEditorModeToolkit::OnMerge),
 // 		FCanExecuteAction::CreateLambda([] { return false; })
 // 		);
 
 	ToolkitCommands->MapAction(
-		Commands.MoveUp, 
+		Commands.MoveUp,
 		FExecuteAction::CreateSP(this, &FFractureEditorModeToolkit::OnMoveUp),
 		FCanExecuteAction::CreateLambda([] { return IsGeometryCollectionSelected(); })
 	);
@@ -687,7 +693,7 @@ bool FFractureEditorModeToolkit::GetShowBoneColors() const
 	{
 		UGeometryCollectionComponent* Comp = GeomCompSelection.Array()[0];
 		FScopedColorEdit EditBoneColor = Comp->EditBoneSelection();
-		return EditBoneColor.GetShowBoneColors();	
+		return EditBoneColor.GetShowBoneColors();
 	}
 
 	return false;
@@ -741,7 +747,7 @@ void FFractureEditorModeToolkit::OnSetExplodedViewValue(float NewValue)
 				}
 			}
 		}
-		
+
 		GCurrentLevelEditingViewportClient->Invalidate();
 	}
 }
@@ -834,7 +840,7 @@ TSharedRef<SWidget> FFractureEditorModeToolkit::GetAutoClusterModesMenu()
 
 	MenuBuilder.AddMenuEntry(
 		LOCTEXT("AutoClusterModeBoundingBox", "Bounding Box"),
-		LOCTEXT("AutoClusterModeBoundingBoxTooltip", "Use the Bounding Box method when Auto Clustering."),
+		LOCTEXT("AutoClusterModeBoundingBoxTooltip", "Use the Bounding Box method when Auto Clustering.  This uses the overlapping of the bones bounding boxes to determine which bones to connect."),
 		FSlateIcon(),
 		FUIAction(
 			FExecuteAction::CreateSP( this, &FFractureEditorModeToolkit::SetAutoClusterMode, EFractureAutoClusterMode::BoundingBox),
@@ -847,7 +853,7 @@ TSharedRef<SWidget> FFractureEditorModeToolkit::GetAutoClusterModesMenu()
 
 	MenuBuilder.AddMenuEntry(
 		LOCTEXT("AutoClusterModeProximity", "Proximity"),
-		LOCTEXT("AutoClusterModeProximityTooltip", "Use the Proximity method when Auto Clustering."),
+		LOCTEXT("AutoClusterModeProximityTooltip", "Use the Proximity method when Auto Clustering.  This uses actual bone edge connectivity to determine which bones to connect."),
 		FSlateIcon(),
 		FUIAction(
 			FExecuteAction::CreateSP( this, &FFractureEditorModeToolkit::SetAutoClusterMode, EFractureAutoClusterMode::Proximity),
@@ -860,7 +866,7 @@ TSharedRef<SWidget> FFractureEditorModeToolkit::GetAutoClusterModesMenu()
 
 	MenuBuilder.AddMenuEntry(
 		LOCTEXT("AutoClusterModeDistance", "Distance"),
-		LOCTEXT("AutoClusterModeDistanceTooltip", "Use the Distance method when Auto Clustering."),
+		LOCTEXT("AutoClusterModeDistanceTooltip", "Use the Distance method when Auto Clustering.  This uses a simple distance calculation to determine which bones to connect"),
 		FSlateIcon(),
 		FUIAction(
 			FExecuteAction::CreateSP( this, &FFractureEditorModeToolkit::SetAutoClusterMode, EFractureAutoClusterMode::Distance),
@@ -888,7 +894,7 @@ FText FFractureEditorModeToolkit::GetAutoClusterModeDisplayName() const
 	{
 		case EFractureAutoClusterMode::Proximity: return LOCTEXT("AutoClusterProximity", "Proximity");
 		case EFractureAutoClusterMode::Distance: return LOCTEXT("AutoClusterDistance", "Distance");
-		case EFractureAutoClusterMode::BoundingBox: 
+		case EFractureAutoClusterMode::BoundingBox:
 		default:
 			return LOCTEXT("AutoClusterBoundingBox", "Bounding Box");
 	}
@@ -927,8 +933,8 @@ TSharedRef<SWidget> FFractureEditorModeToolkit::GetLevelViewMenuContent()
 	for (int32 i = 0; i < GetLevelCount(); i++)
 	{
 		MenuBuilder.AddMenuEntry(
-			FText::Format(LOCTEXT("LevelMenuN", "Level {0}"), FText::AsNumber(i)), 
-			FText::Format(LOCTEXT("LevelMenuNTooltip", "View Level {0} in this Geometry Collecdtion"), FText::AsNumber(i)), 
+			FText::Format(LOCTEXT("LevelMenuN", "Level {0}"), FText::AsNumber(i)),
+			FText::Format(LOCTEXT("LevelMenuNTooltip", "View Level {0} in this Geometry Collecdtion"), FText::AsNumber(i)),
 			FSlateIcon(),
 			FUIAction(
 				FExecuteAction::CreateSP(this, &FFractureEditorModeToolkit::OnSetLevelViewValue, i),
@@ -947,7 +953,8 @@ TSharedRef<SWidget> FFractureEditorModeToolkit::GetViewMenuContent()
 	const FFractureEditorCommands& Commands = FFractureEditorCommands::Get();
 
 	FMenuBuilder MenuBuilder(false, GetToolkitCommands());
-	MenuBuilder.AddMenuEntry(Commands.ToggleShowBoneColors);
+// 	MenuBuilder.AddMenuEntry(Commands.ToggleShowBoneColors, NAME_None, TEXT("Key"), TEXT("Key tooltip"));
+	MenuBuilder.AddMenuEntry(Commands.ToggleShowBoneColors, NAME_None, LOCTEXT("AAAAAAAAAAAAAAA", "BBBBBBBBBBBBBBBB"), LOCTEXT("CCCCCCCCCCCC", "DDDDDDDDDDDDDDDDDD"), FSlateIcon(FEditorStyle::GetStyleSetName(), "WorldBrowser.DirectionXNegative"));
 
 
 	return MenuBuilder.MakeWidget();
@@ -1037,7 +1044,7 @@ static bool CanFractureComponent(UPrimitiveComponent* PrimitiveComponent)
 			(PrimitiveComponent->GetOwner() == nullptr || !PrimitiveComponent->GetOwner()->IsEditorOnly());
 }
 
-void FFractureEditorModeToolkit::OnSelectByMode(GeometryCollection::ESelectionMode SelectionMode) 
+void FFractureEditorModeToolkit::OnSelectByMode(GeometryCollection::ESelectionMode SelectionMode)
 {
 	USelection* SelectionSet = GEditor->GetSelectedActors();
 
@@ -1143,7 +1150,7 @@ void FFractureEditorModeToolkit::GetFractureContexts(TArray<FFractureContext>& F
 							{
 								continue;
 							}
-							
+
 							if (TransformToGeometryIndex[BoneIndex] > -1)
 							{
 								FractureContext.Bounds += BoundsToBone[BoneIndex];
@@ -1303,7 +1310,7 @@ ULevel* FFractureEditorModeToolkit::GetSelectedLevel()
 }
 
 
-void FFractureEditorModeToolkit::GetSelectedGeometryCollectionComponents(TSet<UGeometryCollectionComponent*>& GeomCompSelection) 
+void FFractureEditorModeToolkit::GetSelectedGeometryCollectionComponents(TSet<UGeometryCollectionComponent*>& GeomCompSelection)
 {
 	USelection* SelectionSet = GEditor->GetSelectedActors();
 	TArray<AActor*> SelectedActors;
@@ -1661,7 +1668,7 @@ AActor* FFractureEditorModeToolkit::AddActor(ULevel* InLevel, UClass* Class)
 class AGeometryCollectionActor* FFractureEditorModeToolkit::CreateNewGeometryActor(const FString& InAssetPath, const FTransform& Transform, bool AddMaterials /*= false*/)
 {
 
-	FString UniquePackageName = InAssetPath; 
+	FString UniquePackageName = InAssetPath;
 	FString UniqueAssetName = FPackageName::GetLongPackageAssetName(InAssetPath);
 
 	FAssetToolsModule& AssetToolsModule = FModuleManager::LoadModuleChecked<FAssetToolsModule>("AssetTools");
@@ -1680,7 +1687,7 @@ class AGeometryCollectionActor* FFractureEditorModeToolkit::CreateNewGeometryAct
 	// copy transform of original static mesh actor to this new actor
 	NewActor->SetActorLabel(UniqueAssetName);
 	NewActor->SetActorTransform(Transform);
-	
+
 	// Mark relevant stuff dirty
 	FAssetRegistryModule::AssetCreated(InGeometryCollection);
 	InGeometryCollection->MarkPackageDirty();
@@ -1907,7 +1914,7 @@ AGeometryCollectionActor*  FFractureEditorModeToolkit::ConvertStaticMeshToGeomet
 		const FVector ActorOffset(Actor->GetActorLocation() - FirstActor->GetActorLocation());
 
 		check(FracturedGeometryCollection);
- 
+
 		TArray<UStaticMeshComponent*> StaticMeshComponents;
 		Actor->GetComponents<UStaticMeshComponent>(StaticMeshComponents, true);
 		for (int32 ii = 0, ni = StaticMeshComponents.Num(); ii < ni; ++ii)
@@ -1922,7 +1929,7 @@ AGeometryCollectionActor*  FFractureEditorModeToolkit::ConvertStaticMeshToGeomet
 				FGeometryCollectionConversion::AppendStaticMesh(ComponentStaticMesh, StaticMeshComponent, ComponentTranform, FracturedGeometryCollection, true);
 			}
 		}
- 
+
 		FracturedGeometryCollection->InitializeMaterials();
 	}
 	AddSingleRootNodeIfRequired(FracturedGeometryCollection);
diff --git a/Engine/Plugins/Experimental/ChaosEditor/Source/FractureEditor/Private/FractureToolUniform.h b/Engine/Plugins/Experimental/ChaosEditor/Source/FractureEditor/Private/FractureToolUniform.h
--- a/Engine/Plugins/Experimental/ChaosEditor/Source/FractureEditor/Private/FractureToolUniform.h
+++ b/Engine/Plugins/Experimental/ChaosEditor/Source/FractureEditor/Private/FractureToolUniform.h
@@ -18,7 +18,7 @@ public:
 	virtual void PostEditChangeProperty(FPropertyChangedEvent& PropertyChangedEvent) override;
 	virtual void PostEditChangeChainProperty(struct FPropertyChangedChainEvent& PropertyChangedEvent) override;
 #endif
-	   
+
 	UFractureUniformSettings()
 		: NumberVoronoiSitesMin(20)
 		, NumberVoronoiSitesMax(20)
@@ -26,10 +26,11 @@ public:
 	{}
 
 
-	/** Number of Voronoi sites - Uniform Voronoi Method */
+	/** Minimum Number of Voronoi sites - A random number will be chosen between the Min and Max for each bone you have selected */
 	UPROPERTY(EditAnywhere, Category = UniformVoronoi, meta = (DisplayName = "Minimum Voronoi Sites", UIMin = "1", UIMax = "5000", ClampMin = "1"))
 	int32 NumberVoronoiSitesMin;
 
+	/** Maximum Number of Voronoi sites - A random number will be chosen between the Min and Max for each bone you have selected */
 	UPROPERTY(EditAnywhere, Category = UniformVoronoi, meta = (DisplayName = "Maximum Voronoi Sites", UIMin = "1", UIMax = "5000", ClampMin = "1"))
 	int32 NumberVoronoiSitesMax;
 
