Added a bunch of tooltips to fracture noise.  #jira none #rb none #fyi jack.oakman

diff --git a/Engine/Plugins/Experimental/ChaosEditor/Source/FractureEditor/Public/FractureTool.h b/Engine/Plugins/Experimental/ChaosEditor/Source/FractureEditor/Public/FractureTool.h
--- a/Engine/Plugins/Experimental/ChaosEditor/Source/FractureEditor/Public/FractureTool.h
+++ b/Engine/Plugins/Experimental/ChaosEditor/Source/FractureEditor/Public/FractureTool.h
@@ -65,15 +65,19 @@ public:
 	UPROPERTY(EditAnywhere, Category = CommonFracture, meta = (DisplayName = "Draw Diagram"))
 	bool bDrawDiagram;
 
+	/** Size of the noise displacement in centimeters */
 	UPROPERTY(EditAnywhere, Category = Noise, meta = (UIMin = "0.0"))
 	float Amplitude;
 
+	/** Period of the Perlin noise.  Smaller values will create noise faces that are smoother */
 	UPROPERTY(EditAnywhere, Category = Noise)
 	float Frequency;
 
+	/** Number of fractal layers of Perlin noise to apply.  Smaller values (1 or 2) will create noise that looks like gentle rolling hills, while larger values (> 4) will tend to look more like craggy mountains */
 	UPROPERTY(EditAnywhere, Category = Noise, meta = (UIMin = "1"))
 	int32 OctaveNumber;
 
+	/** Spacing between vertices on cut surfaces, where noise is added.  Larger spacing between vertices will create more efficient meshes with fewer triangles, but less resolution to see the shape of the added noise  */
 	UPROPERTY(EditAnywhere, Category = Noise, meta = (UIMin = "1"))
 	int32 SurfaceResolution;
 
