remove errant tooltip #jira none #rb trivial #fyi mason.seay

diff --git a/Engine/Plugins/Experimental/ChaosEditor/Source/FractureEditor/Private/FractureEditorModeToolkit.cpp b/Engine/Plugins/Experimental/ChaosEditor/Source/FractureEditor/Private/FractureEditorModeToolkit.cpp
--- a/Engine/Plugins/Experimental/ChaosEditor/Source/FractureEditor/Private/FractureEditorModeToolkit.cpp
+++ b/Engine/Plugins/Experimental/ChaosEditor/Source/FractureEditor/Private/FractureEditorModeToolkit.cpp
@@ -949,16 +949,12 @@ TSharedRef<SWidget> FFractureEditorModeToolkit::GetLevelViewMenuContent()
 
 TSharedRef<SWidget> FFractureEditorModeToolkit::GetViewMenuContent()
 {
-
 	const FFractureEditorCommands& Commands = FFractureEditorCommands::Get();
 
 	FMenuBuilder MenuBuilder(false, GetToolkitCommands());
-// 	MenuBuilder.AddMenuEntry(Commands.ToggleShowBoneColors, NAME_None, TEXT("Key"), TEXT("Key tooltip"));
-	MenuBuilder.AddMenuEntry(Commands.ToggleShowBoneColors, NAME_None, LOCTEXT("AAAAAAAAAAAAAAA", "BBBBBBBBBBBBBBBB"), LOCTEXT("CCCCCCCCCCCC", "DDDDDDDDDDDDDDDDDD"), FSlateIcon(FEditorStyle::GetStyleSetName(), "WorldBrowser.DirectionXNegative"));
-
+ 	MenuBuilder.AddMenuEntry(Commands.ToggleShowBoneColors);
 
 	return MenuBuilder.MakeWidget();
-
 }
 
 void FFractureEditorModeToolkit::SetActiveTool(UFractureTool* InActiveTool)
