Fix editor crashing when selecting Geometry Collection after deleting RestCollection asset from context browser.

diff --git a/Engine/Plugins/Experimental/ChaosEditor/Source/FractureEditor/Private/FractureEditorModeToolkit.cpp b/Engine/Plugins/Experimental/ChaosEditor/Source/FractureEditor/Private/FractureEditorModeToolkit.cpp
--- a/Engine/Plugins/Experimental/ChaosEditor/Source/FractureEditor/Private/FractureEditorModeToolkit.cpp
+++ b/Engine/Plugins/Experimental/ChaosEditor/Source/FractureEditor/Private/FractureEditorModeToolkit.cpp
@@ -992,10 +992,14 @@ void FFractureEditorModeToolkit::SetOutlinerComponents(const TArray<UGeometryCol
 	{
 		FGeometryCollectionEdit RestCollection = Component->EditRestCollection(GeometryCollection::EEditUpdate::None);
 		UGeometryCollection* FracturedGeometryCollection = RestCollection.GetRestCollection();
-		TSharedPtr<FGeometryCollection, ESPMode::ThreadSafe> GeometryCollectionPtr = FracturedGeometryCollection->GetGeometryCollection();
 
-		FGeometryCollectionClusteringUtility::UpdateHierarchyLevelOfChildren(GeometryCollectionPtr.Get(), -1);
-		UpdateExplodedVectors(Component);
+		if (FracturedGeometryCollection) // Prevents crash when GC is deleted from content browser and actor is selected.
+		{
+			TSharedPtr<FGeometryCollection, ESPMode::ThreadSafe> GeometryCollectionPtr = FracturedGeometryCollection->GetGeometryCollection();
+
+			FGeometryCollectionClusteringUtility::UpdateHierarchyLevelOfChildren(GeometryCollectionPtr.Get(), -1);
+			UpdateExplodedVectors(Component);
+		}
 	}
 
 	if (OutlinerView)
@@ -1085,6 +1089,10 @@ void FFractureEditorModeToolkit::GetFractureContexts(TArray<FFractureContext>& F
 				{
 					FGeometryCollectionEdit RestCollection = GeometryCollectionComponent->EditRestCollection(GeometryCollection::EEditUpdate::None);
 					UGeometryCollection* FracturedGeometryCollection = RestCollection.GetRestCollection();
+					if (FracturedGeometryCollection == nullptr)
+					{
+						continue;
+					}
 
 					const TArray<int32>& SelectedBonesOriginal = GeometryCollectionComponent->GetSelectedBones();
 
