Fix details view not showing Geometry Collection details after creating GC.

diff --git a/Engine/Plugins/Experimental/ChaosEditor/Source/FractureEditor/Private/FractureEditorModeToolkit.cpp b/Engine/Plugins/Experimental/ChaosEditor/Source/FractureEditor/Private/FractureEditorModeToolkit.cpp
--- a/Engine/Plugins/Experimental/ChaosEditor/Source/FractureEditor/Private/FractureEditorModeToolkit.cpp
+++ b/Engine/Plugins/Experimental/ChaosEditor/Source/FractureEditor/Private/FractureEditorModeToolkit.cpp
@@ -1615,6 +1615,8 @@ void FFractureEditorModeToolkit::OnGenerateAssetPathChosen(const FString& InAsse
 			const FVector ActorLocation(FirstActor->GetActorLocation());
 			GeometryCollectionActor->SetActorLocation(ActorLocation);
 
+			// Clear selection of mesh actor used to make GC before selecting, will cause details pane to not display geometry collection details.
+			GEditor->SelectNone(true, true, false);
 			GEditor->SelectActor(GeometryCollectionActor, true, true);
 
 			EditBoneColor.SelectBones(GeometryCollection::ESelectionMode::AllGeometry);
