Added transaction and modify when Selecting bones in fracture editor outliner. #jira UE-78017 #lockdown cristina.riveron #rb jimmy.andrws #fyi michael.lentine #fyi louise.rasmussen

diff --git a/Engine/Plugins/Experimental/ChaosEditor/Source/FractureEditor/Private/SGeometryCollectionOutliner.cpp b/Engine/Plugins/Experimental/ChaosEditor/Source/FractureEditor/Private/SGeometryCollectionOutliner.cpp
--- a/Engine/Plugins/Experimental/ChaosEditor/Source/FractureEditor/Private/SGeometryCollectionOutliner.cpp
+++ b/Engine/Plugins/Experimental/ChaosEditor/Source/FractureEditor/Private/SGeometryCollectionOutliner.cpp
@@ -10,6 +10,7 @@
 #include "LevelEditor.h"
 #include "Modules/ModuleManager.h"
 #include "GeometryCollection/GeometryCollectionComponent.h"
+#include "FractureEditorMode.h"
 
 #define LOCTEXT_NAMESPACE "ChaosEditor"
 
@@ -186,12 +187,14 @@ void SGeometryCollectionOutliner::OnSelectionChanged(FGeometryCollectionTreeItem
 		FGeometryCollectionTreeItemList SelectedItems;
 		TreeView->GetSelectedItems(SelectedItems);
 		
+		FScopedTransaction Transaction(FractureTransactionContexts::SelectBoneContext, LOCTEXT("SelectGeometryCollectionBoneTransaction", "Select Bone"), Item->GetComponent());
 		for(auto& SelectedItem : SelectedItems)
 		{
 			if (SelectedItem->GetBoneIndex() != INDEX_NONE)
 			{
 				TArray<int32>& SelectedBones = ComponentToBoneSelectionMap.FindChecked(Item->GetComponent());
 				SelectedBones.Add(SelectedItem->GetBoneIndex());
+				Item->GetComponent()->Modify();
 			}
 		}
 		// Fire off the delegate for each component
