Fix SGeometryCollectionOutliner
#rnx

diff --git a/Engine/Plugins/Experimental/ChaosEditor/Source/FractureEditor/Private/SGeometryCollectionOutliner.cpp b/Engine/Plugins/Experimental/ChaosEditor/Source/FractureEditor/Private/SGeometryCollectionOutliner.cpp
--- a/Engine/Plugins/Experimental/ChaosEditor/Source/FractureEditor/Private/SGeometryCollectionOutliner.cpp
+++ b/Engine/Plugins/Experimental/ChaosEditor/Source/FractureEditor/Private/SGeometryCollectionOutliner.cpp
@@ -11,6 +11,7 @@
 #include "Modules/ModuleManager.h"
 #include "GeometryCollection/GeometryCollectionComponent.h"
 #include "FractureEditorMode.h"
+#include "ScopedTransaction.h"
 
 #define LOCTEXT_NAMESPACE "ChaosEditor"
 
