Fracture button enable check now makes sure selected actors are in the editor world.  #jira UE-78133 #rb michael.lentine #lockdown cristina.riveron

diff --git a/Engine/Plugins/Experimental/ChaosEditor/Source/FractureEditor/Private/FractureEditorModeToolkit.cpp b/Engine/Plugins/Experimental/ChaosEditor/Source/FractureEditor/Private/FractureEditorModeToolkit.cpp
--- a/Engine/Plugins/Experimental/ChaosEditor/Source/FractureEditor/Private/FractureEditorModeToolkit.cpp
+++ b/Engine/Plugins/Experimental/ChaosEditor/Source/FractureEditor/Private/FractureEditorModeToolkit.cpp
@@ -1252,6 +1252,11 @@ bool FFractureEditorModeToolkit::CanExecuteFracture() const
 		return false;
 	}
 
+	if (!IsSeletedActorsInEditorWorld())
+	{
+		return false;
+	}
+
 	if (IsGeometryCollectionSelected())
 	{
 		return IsLeafBoneSelected();
@@ -1261,7 +1266,7 @@ bool FFractureEditorModeToolkit::CanExecuteFracture() const
 	{
 		return false;
 	}
-
+	
 	return false;
 }
 
@@ -1764,10 +1769,25 @@ bool FFractureEditorModeToolkit::IsStaticMeshSelected()
 		}
 	}
 	return false;
-
 }
 
-
+bool FFractureEditorModeToolkit::IsSeletedActorsInEditorWorld()
+{
+	USelection* SelectedActors = GEditor->GetSelectedActors();
+	for (FSelectionIterator Iter(*SelectedActors); Iter; ++Iter)
+	{
+		AActor* Actor = Cast<AActor>(*Iter);
+		if (Actor)
+		{
+			check(Actor->GetWorld());
+			if (Actor->GetWorld()->WorldType != EWorldType::Editor)
+			{
+				return false;
+			}
+		}
+	}
+	return true;
+}
 
 bool GetValidGeoCenter(const TManagedArray<int32>& TransformToGeometryIndex, const TArray<FTransform>& Transforms, const TManagedArray<TSet<int32>>& Children, const TManagedArray<FBox>& BoundingBox, int32 TransformIndex, FVector& OutGeoCenter )
 {
diff --git a/Engine/Plugins/Experimental/ChaosEditor/Source/FractureEditor/Public/FractureEditorModeToolkit.h b/Engine/Plugins/Experimental/ChaosEditor/Source/FractureEditor/Public/FractureEditorModeToolkit.h
--- a/Engine/Plugins/Experimental/ChaosEditor/Source/FractureEditor/Public/FractureEditorModeToolkit.h
+++ b/Engine/Plugins/Experimental/ChaosEditor/Source/FractureEditor/Public/FractureEditorModeToolkit.h
@@ -115,6 +115,7 @@ public:
 protected:
 	static bool IsGeometryCollectionSelected();
 	static bool IsStaticMeshSelected();
+	static bool IsSeletedActorsInEditorWorld();
 	void UpdateExplodedVectors(UGeometryCollectionComponent* GeometryCollectionComponent) const;
 private:
 	void OnOutlinerBoneSelectionChanged(UGeometryCollectionComponent* RootComponent, TArray<int32>& SelectedBones);
