[NO NON-PII COMMIT MESSAGE]

diff --git a/Engine/Plugins/Experimental/ChaosEditor/Source/FractureEditor/Private/SGeometryCollectionOutliner.cpp b/Engine/Plugins/Experimental/ChaosEditor/Source/FractureEditor/Private/SGeometryCollectionOutliner.cpp
--- a/Engine/Plugins/Experimental/ChaosEditor/Source/FractureEditor/Private/SGeometryCollectionOutliner.cpp
+++ b/Engine/Plugins/Experimental/ChaosEditor/Source/FractureEditor/Private/SGeometryCollectionOutliner.cpp
@@ -105,13 +105,18 @@ void SGeometryCollectionOutliner::UpdateGeometryCollection()
 
 void SGeometryCollectionOutliner::SetComponents(const TArray<UGeometryCollectionComponent*>& InNewComponents)
 {
+	// Clear the cached Tree ItemSelection without affecting the SelectedBones as 
+	// we want to refresh the tree selection using selected bones
+	TGuardValue<bool> ExternalSelectionGuard(bPerformingSelection, true);
+	TreeView->ClearSelection();
+
 	RootNodes.Empty();
 
 	for (UGeometryCollectionComponent* Component : InNewComponents)
 	{
 		RootNodes.Add(MakeShared<FGeometryCollectionTreeItemComponent>(Component));
 		TArray<int32> SelectedBones = Component->GetSelectedBones();
-		SetBoneSelection(Component, SelectedBones, !SelectedBones.Num());
+		SetBoneSelection(Component, SelectedBones, false);
 	}
 
 	TreeView->RequestTreeRefresh();
