[NO NON-PII COMMIT MESSAGE]

diff --git a/Engine/Plugins/Experimental/ChaosEditor/Source/FractureEditor/Private/FractureSelectionTools.cpp b/Engine/Plugins/Experimental/ChaosEditor/Source/FractureEditor/Private/FractureSelectionTools.cpp
--- a/Engine/Plugins/Experimental/ChaosEditor/Source/FractureEditor/Private/FractureSelectionTools.cpp
+++ b/Engine/Plugins/Experimental/ChaosEditor/Source/FractureEditor/Private/FractureSelectionTools.cpp
@@ -52,6 +52,7 @@ void FFractureSelectionTools::ClearSelectedBones(UGeometryCollectionComponent* G
 {
 	FScopedColorEdit EditBoneColor = GeometryCollectionComponent->EditBoneSelection();
 	EditBoneColor.ResetBoneSelection();
+	EditBoneColor.ResetHighlightedBones();
 }
 
 
diff --git a/Engine/Plugins/Experimental/ChaosEditor/Source/FractureEditor/Private/SGeometryCollectionOutliner.cpp b/Engine/Plugins/Experimental/ChaosEditor/Source/FractureEditor/Private/SGeometryCollectionOutliner.cpp
--- a/Engine/Plugins/Experimental/ChaosEditor/Source/FractureEditor/Private/SGeometryCollectionOutliner.cpp
+++ b/Engine/Plugins/Experimental/ChaosEditor/Source/FractureEditor/Private/SGeometryCollectionOutliner.cpp
@@ -190,17 +190,22 @@ void SGeometryCollectionOutliner::OnSelectionChanged(FGeometryCollectionTreeItem
 			ComponentToBoneSelectionMap.Add(Root->GetComponent(), TArray<int32>());
 		}
 
+		if (Item == nullptr)
+		{
+			TreeView->ClearSelection();
+		}
+
 		FGeometryCollectionTreeItemList SelectedItems;
 		TreeView->GetSelectedItems(SelectedItems);
-		
-		FScopedTransaction Transaction(FractureTransactionContexts::SelectBoneContext, LOCTEXT("SelectGeometryCollectionBoneTransaction", "Select Bone"), Item->GetComponent());
+
+		FScopedTransaction Transaction(FractureTransactionContexts::SelectBoneContext, LOCTEXT("SelectGeometryCollectionBoneTransaction", "Select Bone"), Item != nullptr ? Item->GetComponent() : nullptr);
 		for(auto& SelectedItem : SelectedItems)
 		{
 			if (SelectedItem->GetBoneIndex() != INDEX_NONE)
 			{
-				TArray<int32>& SelectedBones = ComponentToBoneSelectionMap.FindChecked(Item->GetComponent());
+				TArray<int32>& SelectedBones = ComponentToBoneSelectionMap.FindChecked(SelectedItem->GetComponent());
 				SelectedBones.Add(SelectedItem->GetBoneIndex());
-				Item->GetComponent()->Modify();
+				SelectedItem->GetComponent()->Modify();
 			}
 		}
 		// Fire off the delegate for each component
