Fixup engine code misusing GetComponentsByClass
#rnx

diff --git a/Engine/Plugins/Editor/MeshEditor/Source/MeshEditor/FractureToolComponent.cpp b/Engine/Plugins/Editor/MeshEditor/Source/MeshEditor/FractureToolComponent.cpp
--- a/Engine/Plugins/Editor/MeshEditor/Source/MeshEditor/FractureToolComponent.cpp
+++ b/Engine/Plugins/Editor/MeshEditor/Source/MeshEditor/FractureToolComponent.cpp
@@ -207,14 +207,11 @@ void UFractureToolComponent::OnEnterFractureMode()
 	bInFractureMode = true;
 	for(const AActor* SelectedActor : GetSelectedActors())
 	{
-		TArray<UActorComponent*> PrimitiveComponents = SelectedActor->GetComponentsByClass(UPrimitiveComponent::StaticClass());
-		for (UActorComponent* PrimitiveActorComponent : PrimitiveComponents)
+		TInlineComponentArray<UPrimitiveComponent*> PrimitiveComponents;
+		SelectedActor->GetComponents(PrimitiveComponents);
+		for (UPrimitiveComponent* Component : PrimitiveComponents)
 		{
-			UPrimitiveComponent* Component = CastChecked<UPrimitiveComponent>(PrimitiveActorComponent);
-			if (Component)
-			{
-				OnSelected(Component);
-			}
+			OnSelected(Component);
 		}
 	}
 }
@@ -225,14 +222,11 @@ void UFractureToolComponent::OnExitFractureMode()
 	// find all the selected geometry collections and turn off color rendering mode
 	for (const AActor* SelectedActor : GetSelectedActors())
 	{
-		TArray<UActorComponent*> PrimitiveComponents = SelectedActor->GetComponentsByClass(UPrimitiveComponent::StaticClass());
-		for (UActorComponent* PrimitiveActorComponent : PrimitiveComponents)
+		TInlineComponentArray<UPrimitiveComponent*> PrimitiveComponents;
+		SelectedActor->GetComponents(PrimitiveComponents);
+		for (UPrimitiveComponent* Component : PrimitiveComponents)
 		{
-			UPrimitiveComponent* Component = CastChecked<UPrimitiveComponent>(PrimitiveActorComponent);
-			if (Component)
-			{
-				OnDeselected(Component);
-			}
+			OnDeselected(Component);
 		}
 	}
 	LeaveFracturingCleanup();
@@ -258,12 +252,12 @@ AActor* UFractureToolComponent::GetEditableMeshActor(UEditableMesh* EditableMesh
 	AActor* ReturnActor = nullptr;
 	const TArray<AActor*>& SelectedActors = GetSelectedActors();
 
-	for (int i = 0; i < SelectedActors.Num(); i++)
+	TInlineComponentArray<UPrimitiveComponent*> PrimitiveComponents;
+	for (int32 i = 0; i < SelectedActors.Num(); i++)
 	{
-		TArray<UActorComponent*> PrimitiveComponents = SelectedActors[i]->GetComponentsByClass(UPrimitiveComponent::StaticClass());
-		for (UActorComponent* PrimitiveActorComponent : PrimitiveComponents)
+		SelectedActors[i]->GetComponents(PrimitiveComponents);
+		for (UPrimitiveComponent* Component : PrimitiveComponents)
 		{
-			UPrimitiveComponent* Component = CastChecked<UPrimitiveComponent>(PrimitiveActorComponent);
 			FEditableMeshSubMeshAddress SubMeshAddress = UEditableMeshFactory::MakeSubmeshAddress(Component, 0);
 
 			if (EditableMesh->GetSubMeshAddress() == SubMeshAddress)
@@ -272,6 +266,7 @@ AActor* UFractureToolComponent::GetEditableMeshActor(UEditableMesh* EditableMesh
 				break;
 			}
 		}
+		PrimitiveComponents.Reset();
 	}
 
 	return ReturnActor;
@@ -282,18 +277,12 @@ AActor* UFractureToolComponent::GetEditableMeshActor()
 	AActor* ReturnActor = nullptr;
 	const TArray<AActor*>& SelectedActors = GetSelectedActors();
 
-	for (int i = 0; i < SelectedActors.Num(); i++)
+	for (int32 i = 0; i < SelectedActors.Num(); i++)
 	{
-		TArray<UActorComponent*> PrimitiveComponents = SelectedActors[i]->GetComponentsByClass(UPrimitiveComponent::StaticClass());
-		for (UActorComponent* PrimitiveActorComponent : PrimitiveComponents)
+		if (UPrimitiveComponent* Component = SelectedActors[i]->FindComponentByClass<UPrimitiveComponent>())
 		{
-			UPrimitiveComponent* Component = CastChecked<UPrimitiveComponent>(PrimitiveActorComponent);
-
-			if (Component)
-			{
-				ReturnActor = Component->GetOwner();
-				break;
-			}
+			ReturnActor = Component->GetOwner();
+			break;
 		}
 	}
 
@@ -351,19 +340,16 @@ void UFractureToolComponent::LeaveFracturingCleanup()
 		if (GeometryActor)
 		{
 			// hide the bones
-			TArray<UActorComponent*> PrimitiveComponents = Actor->GetComponentsByClass(UPrimitiveComponent::StaticClass());
-			for (UActorComponent* PrimitiveActorComponent : PrimitiveComponents)
+			TInlineComponentArray<UPrimitiveComponent*> PrimitiveComponents;
+			Actor->GetComponents(PrimitiveComponents);
+			for (UPrimitiveComponent* Component : PrimitiveComponents)
 			{
-				UPrimitiveComponent* Component = CastChecked<UPrimitiveComponent>(PrimitiveActorComponent);
-				if (Component)
+				UGeometryCollectionComponent* GeometryCollectionComponent = Cast<UGeometryCollectionComponent>(Component);
+				if (GeometryCollectionComponent)
 				{
-					UGeometryCollectionComponent* GeometryCollectionComponent = Cast<UGeometryCollectionComponent>(Component);
-					if (GeometryCollectionComponent)
-					{
-						FScopedColorEdit EditBoneColor = GeometryCollectionComponent->EditBoneSelection();
-						EditBoneColor.SetShowBoneColors(false);
-						EditBoneColor.SetEnableBoneSelection(false);
-					}
+					FScopedColorEdit EditBoneColor = GeometryCollectionComponent->EditBoneSelection();
+					EditBoneColor.SetShowBoneColors(false);
+					EditBoneColor.SetEnableBoneSelection(false);
 				}
 			}
 
diff --git a/Engine/Plugins/Editor/MeshEditor/Source/MeshEditor/MeshEditorMode.cpp b/Engine/Plugins/Editor/MeshEditor/Source/MeshEditor/MeshEditorMode.cpp
--- a/Engine/Plugins/Editor/MeshEditor/Source/MeshEditor/MeshEditorMode.cpp
+++ b/Engine/Plugins/Editor/MeshEditor/Source/MeshEditor/MeshEditorMode.cpp
@@ -4487,11 +4487,10 @@ void FMeshEditorMode::UpdateSelectedEditableMeshes()
 		AActor* Actor = Cast<AActor>( *SelectionIt );
 		if( Actor != nullptr )
 		{
-			TArray<UActorComponent*> PrimitiveComponents = Actor->GetComponentsByClass( UPrimitiveComponent::StaticClass() );
-			for( UActorComponent* PrimitiveActorComponent : PrimitiveComponents )
+			TInlineComponentArray<UPrimitiveComponent*> PrimitiveComponents;
+			Actor->GetComponents(PrimitiveComponents);
+			for(UPrimitiveComponent* Component : PrimitiveComponents )
 			{
-				UPrimitiveComponent* Component = CastChecked<UPrimitiveComponent>( PrimitiveActorComponent );
-
 				// Don't bother with editor-only 'helper' actors, we never want to visualize or edit geometry on those
 				if( !Component->IsEditorOnly() &&
 					Component->GetCollisionEnabled() != ECollisionEnabled::NoCollision &&
diff --git a/Engine/Plugins/Editor/MeshEditor/Source/PolygonModeling/GeometryCollectionCommandCommon.cpp b/Engine/Plugins/Editor/MeshEditor/Source/PolygonModeling/GeometryCollectionCommandCommon.cpp
--- a/Engine/Plugins/Editor/MeshEditor/Source/PolygonModeling/GeometryCollectionCommandCommon.cpp
+++ b/Engine/Plugins/Editor/MeshEditor/Source/PolygonModeling/GeometryCollectionCommandCommon.cpp
@@ -245,12 +245,12 @@ AActor* FGeometryCollectionCommandCommon::GetEditableMeshActor(UEditableMesh* Ed
 	AActor* ReturnActor = nullptr;
 	const TArray<AActor*>& SelectedActors = CommandCommon::GetSelectedActors();
 
-	for (int i = 0; i < SelectedActors.Num(); i++)
+	TInlineComponentArray<UPrimitiveComponent*> PrimitiveComponents;
+	for (int32 i = 0; i < SelectedActors.Num(); i++)
 	{
-		TArray<UActorComponent*> PrimitiveComponents = SelectedActors[i]->GetComponentsByClass(UPrimitiveComponent::StaticClass());
-		for (UActorComponent* PrimitiveActorComponent : PrimitiveComponents)
+		SelectedActors[i]->GetComponents(PrimitiveComponents);
+		for (UPrimitiveComponent* Component : PrimitiveComponents)
 		{
-			UPrimitiveComponent* Component = CastChecked<UPrimitiveComponent>(PrimitiveActorComponent);
 			FEditableMeshSubMeshAddress SubMeshAddress = UEditableMeshFactory::MakeSubmeshAddress(Component, 0);
 
 			if (EditableMesh->GetSubMeshAddress() == SubMeshAddress)
@@ -259,6 +259,7 @@ AActor* FGeometryCollectionCommandCommon::GetEditableMeshActor(UEditableMesh* Ed
 				break;
 			}
 		}
+		PrimitiveComponents.Reset();
 	}
 
 	return ReturnActor;
@@ -267,10 +268,10 @@ AActor* FGeometryCollectionCommandCommon::GetEditableMeshActor(UEditableMesh* Ed
 UEditableMesh* FGeometryCollectionCommandCommon::GetEditableMeshForActor(AActor* Actor, TArray<UEditableMesh *>& SelectedMeshes)
 {
 	check(Actor);
-	TArray<UActorComponent*> PrimitiveComponents = Actor->GetComponentsByClass(UPrimitiveComponent::StaticClass());
-	for (UActorComponent* PrimitiveActorComponent : PrimitiveComponents)
+	TInlineComponentArray<UPrimitiveComponent*> PrimitiveComponents;
+	Actor->GetComponents(PrimitiveComponents);
+	for (UPrimitiveComponent* Component : PrimitiveComponents)
 	{
-		UPrimitiveComponent* Component = CastChecked<UPrimitiveComponent>(PrimitiveActorComponent);
 		FEditableMeshSubMeshAddress SubMeshAddress = UEditableMeshFactory::MakeSubmeshAddress(Component, 0);
 
 		for (UEditableMesh* EditableMesh : SelectedMeshes)
@@ -389,12 +390,13 @@ void FGeometryCollectionCommandCommon::AppendMeshesToGeometryCollection(TArray<A
 				TManagedArray<FTransform>& ExplodedTransforms = GeometryCollection->GetAttribute<FTransform>("ExplodedTransform", FGeometryCollection::TransformGroup);
 				TManagedArray<FVector>& ExplodedVectors = GeometryCollection->GetAttribute<FVector>("ExplodedVector", FGeometryCollection::TransformGroup);
 
-				TArray<UActorComponent*> PrimitiveComponents = SelectedActor->GetComponentsByClass(UPrimitiveComponent::StaticClass());
-				for (UActorComponent* PrimitiveComponent : PrimitiveComponents)
+				TInlineComponentArray<UPrimitiveComponent*> PrimitiveComponents;
+				SelectedActor->GetComponents(PrimitiveComponents);
+				for (UPrimitiveComponent* Component : PrimitiveComponents)
 				{
 					bool ValidComponent = false;
 
-					if (UStaticMeshComponent* StaticMeshComp = Cast<UStaticMeshComponent>(PrimitiveComponent))
+					if (UStaticMeshComponent* StaticMeshComp = Cast<UStaticMeshComponent>(Component))
 					{
 						MeshTransform = StaticMeshComp->GetComponentTransform();
 						MeshTransform = MeshTransform.GetRelativeTransform(SourceActorTransform);
@@ -402,7 +404,7 @@ void FGeometryCollectionCommandCommon::AppendMeshesToGeometryCollection(TArray<A
 						FGeometryCollectionConversion::AppendStaticMesh(StaticMeshComp->GetStaticMesh(), StaticMeshComp, MeshTransform, GeometryCollectionObject, false);
 						ValidComponent = true;
 					}
-					else if (UGeometryCollectionComponent* GeometryCollectionComponent = Cast<UGeometryCollectionComponent>(PrimitiveComponent))
+					else if (UGeometryCollectionComponent* GeometryCollectionComponent = Cast<UGeometryCollectionComponent>(Component))
 					{
 						MeshTransform = GeometryCollectionComponent->GetComponentTransform();
 						MeshTransform = MeshTransform.GetRelativeTransform(SourceActorTransform);
diff --git a/Engine/Plugins/Experimental/ChaosEditor/Source/FractureEditor/Private/FractureEditorMode.cpp b/Engine/Plugins/Experimental/ChaosEditor/Source/FractureEditor/Private/FractureEditorMode.cpp
--- a/Engine/Plugins/Experimental/ChaosEditor/Source/FractureEditor/Private/FractureEditorMode.cpp
+++ b/Engine/Plugins/Experimental/ChaosEditor/Source/FractureEditor/Private/FractureEditorMode.cpp
@@ -218,11 +218,11 @@ bool FFractureEditorMode::FrustumSelect(const FConvexVolume& InFrustum, FEditorV
 
 				if (SelectedBonesArray.Num() > 0)
 				{
-					TArray<UActorComponent*> Components = Actor->GetComponentsByClass(UPrimitiveComponent::StaticClass());
-					for (UActorComponent* Component : Components)
+					TInlineComponentArray<UGeometryCollectionComponent*> GeometryCollectionComponents;
+					Actor->GetComponents(GeometryCollectionComponents);
+
+					for (UGeometryCollectionComponent* GeometryCollectionComponent : GeometryCollectionComponents)
 					{
-						UPrimitiveComponent* PrimitiveComponent = CastChecked<UPrimitiveComponent>(Component);
-						UGeometryCollectionComponent* GeometryCollectionComponent = Cast<UGeometryCollectionComponent>(PrimitiveComponent);
 						FScopedColorEdit ColorEdit = GeometryCollectionComponent->EditBoneSelection();
 						ColorEdit.SelectBones(GeometryCollection::ESelectionMode::None);
 						ColorEdit.SetSelectedBones(SelectedBonesArray);
diff --git a/Engine/Plugins/Experimental/ChaosEditor/Source/FractureEditor/Private/FractureEditorModeToolkit.cpp b/Engine/Plugins/Experimental/ChaosEditor/Source/FractureEditor/Private/FractureEditorModeToolkit.cpp
--- a/Engine/Plugins/Experimental/ChaosEditor/Source/FractureEditor/Private/FractureEditorModeToolkit.cpp
+++ b/Engine/Plugins/Experimental/ChaosEditor/Source/FractureEditor/Private/FractureEditorModeToolkit.cpp
@@ -730,10 +730,10 @@ void FFractureEditorModeToolkit::OnSetExplodedViewValue(float NewValue)
 
 		for (AActor* Actor : SelectedActors)
 		{
-			TArray<UActorComponent*> Components = Actor->GetComponentsByClass(UPrimitiveComponent::StaticClass());
-			for (UActorComponent* Component : Components)
+			TInlineComponentArray<UPrimitiveComponent*> Components;
+			Actor->GetComponents(Components);
+			for (UPrimitiveComponent* PrimitiveComponent : Components)
 			{
-				UPrimitiveComponent* PrimitiveComponent = CastChecked<UPrimitiveComponent>(Component);
 				AGeometryCollectionActor* GeometryCollectionActor = Cast<AGeometryCollectionActor>(Actor);
 				if(GeometryCollectionActor)
 				{
@@ -1054,18 +1054,14 @@ void FFractureEditorModeToolkit::OnSelectByMode(GeometryCollection::ESelectionMo
 
 	for (AActor* Actor : SelectedActors)
 	{
-		TArray<UActorComponent*> Components = Actor->GetComponentsByClass(UPrimitiveComponent::StaticClass());
-		for (UActorComponent* Component : Components)
-		{
-			UPrimitiveComponent* PrimitiveComponent = CastChecked<UPrimitiveComponent>(Component);
-			UGeometryCollectionComponent* GeometryCollectionComponent = Cast<UGeometryCollectionComponent>(PrimitiveComponent);
+		TInlineComponentArray<UGeometryCollectionComponent*> GeometryCollectionComponents;
+		Actor->GetComponents(GeometryCollectionComponents);
 
-			if (GeometryCollectionComponent)
-			{
-				FScopedColorEdit EditBoneColor = GeometryCollectionComponent->EditBoneSelection();
-				EditBoneColor.SelectBones(SelectionMode);
-				SetBoneSelection(GeometryCollectionComponent, EditBoneColor.GetSelectedBones(), true);
-			}
+		for (UGeometryCollectionComponent* GeometryCollectionComponent : GeometryCollectionComponents)
+		{
+			FScopedColorEdit EditBoneColor = GeometryCollectionComponent->EditBoneSelection();
+			EditBoneColor.SelectBones(SelectionMode);
+			SetBoneSelection(GeometryCollectionComponent, EditBoneColor.GetSelectedBones(), true);
 		}
 	}
 }
@@ -1081,57 +1077,93 @@ void FFractureEditorModeToolkit::GetFractureContexts(TArray<FFractureContext>& F
 		AActor* Actor = Cast<AActor>(*Iter);
 		if (Actor)
 		{
-			TArray<UActorComponent*> Components = Actor->GetComponentsByClass(UPrimitiveComponent::StaticClass());
-			for (UActorComponent* Component : Components)
+			TInlineComponentArray<UGeometryCollectionComponent*> GeometryCollectionComponents;
+			Actor->GetComponents(GeometryCollectionComponents);
+			for (UGeometryCollectionComponent* GeometryCollectionComponent : GeometryCollectionComponents)
 			{
-				UPrimitiveComponent* PrimitiveComponent = CastChecked<UPrimitiveComponent>(Component);
-				if (UGeometryCollectionComponent* GeometryCollectionComponent = Cast<UGeometryCollectionComponent>(PrimitiveComponent))
+				FGeometryCollectionEdit RestCollection = GeometryCollectionComponent->EditRestCollection(GeometryCollection::EEditUpdate::None);
+				UGeometryCollection* FracturedGeometryCollection = RestCollection.GetRestCollection();
+				if (FracturedGeometryCollection == nullptr)
 				{
-					FGeometryCollectionEdit RestCollection = GeometryCollectionComponent->EditRestCollection(GeometryCollection::EEditUpdate::None);
-					UGeometryCollection* FracturedGeometryCollection = RestCollection.GetRestCollection();
-					if (FracturedGeometryCollection == nullptr)
-					{
-						continue;
-					}
+					continue;
+				}
 
-					const TArray<int32>& SelectedBonesOriginal = GeometryCollectionComponent->GetSelectedBones();
+				const TArray<int32>& SelectedBonesOriginal = GeometryCollectionComponent->GetSelectedBones();
 
-					TSharedPtr<FGeometryCollection, ESPMode::ThreadSafe> GeometryCollectionPtr = FracturedGeometryCollection->GetGeometryCollection();
-					FGeometryCollection* OutGeometryCollection = GeometryCollectionPtr.Get();
+				TSharedPtr<FGeometryCollection, ESPMode::ThreadSafe> GeometryCollectionPtr = FracturedGeometryCollection->GetGeometryCollection();
+				FGeometryCollection* OutGeometryCollection = GeometryCollectionPtr.Get();
 
-					const TManagedArray<TSet<int32>>& Children = OutGeometryCollection->GetAttribute<TSet<int32>>("Children", FGeometryCollection::TransformGroup);
+				const TManagedArray<TSet<int32>>& Children = OutGeometryCollection->GetAttribute<TSet<int32>>("Children", FGeometryCollection::TransformGroup);
 
-					TArray<int32> SelectedBones;
-					SelectedBones.Reserve(SelectedBonesOriginal.Num());
-					for (int32 BoneIndex : SelectedBonesOriginal)
+				TArray<int32> SelectedBones;
+				SelectedBones.Reserve(SelectedBonesOriginal.Num());
+				for (int32 BoneIndex : SelectedBonesOriginal)
+				{
+					if (Children[BoneIndex].Num() == 0)
 					{
-						if (Children[BoneIndex].Num() == 0)
-						{
-							SelectedBones.Add(BoneIndex);
-						}
+						SelectedBones.Add(BoneIndex);
 					}
+				}
+
+				const TManagedArray<FTransform>& Transform = OutGeometryCollection->GetAttribute<FTransform>("Transform", FGeometryCollection::TransformGroup);
+				const TManagedArray<int32>& TransformToGeometryIndex = OutGeometryCollection->GetAttribute<int32>("TransformToGeometryIndex", FGeometryCollection::TransformGroup);
 
-					const TManagedArray<FTransform>& Transform = OutGeometryCollection->GetAttribute<FTransform>("Transform", FGeometryCollection::TransformGroup);
-					const TManagedArray<int32>& TransformToGeometryIndex = OutGeometryCollection->GetAttribute<int32>("TransformToGeometryIndex", FGeometryCollection::TransformGroup);
+				const TManagedArray<FBox>& BoundingBoxes = OutGeometryCollection->GetAttribute<FBox>("BoundingBox", FGeometryCollection::GeometryGroup);
 
-					const TManagedArray<FBox>& BoundingBoxes = OutGeometryCollection->GetAttribute<FBox>("BoundingBox", FGeometryCollection::GeometryGroup);
+				TArray<FTransform> Transforms;
+				GeometryCollectionAlgo::GlobalMatrices(Transform, OutGeometryCollection->Parent, Transforms);
 
-					TArray<FTransform> Transforms;
-					GeometryCollectionAlgo::GlobalMatrices(Transform, OutGeometryCollection->Parent, Transforms);
+				TMap<int32, FBox> BoundsToBone;
 
-					TMap<int32, FBox> BoundsToBone;
+				for (int32 Idx = 0, ni = FracturedGeometryCollection->NumElements(FGeometryCollection::TransformGroup); Idx < ni; ++Idx)
+				{
+					if (TransformToGeometryIndex[Idx] > -1)
+					{
+						ensure(TransformToGeometryIndex[Idx] > -1);
+						BoundsToBone.Add(Idx, BoundingBoxes[TransformToGeometryIndex[Idx]].TransformBy(Transforms[Idx]));
+					}
+				}
 
-					for (int32 Idx = 0, ni = FracturedGeometryCollection->NumElements(FGeometryCollection::TransformGroup); Idx < ni; ++Idx)
+				if (CommonSettings->bGroupFracture)
+				{
+					FractureContexts.AddDefaulted();
+					FFractureContext& FractureContext = FractureContexts.Last();
+					FractureContext.RandomSeed = FMath::Rand();
+					if (CommonSettings->RandomSeed > -1)
 					{
-						if (TransformToGeometryIndex[Idx] > -1)
+						// make sure it's unique for each context if it's specified.
+						FractureContext.RandomSeed = CommonSettings->RandomSeed + FractureContexts.Num();
+					}
+
+					FractureContext.OriginalActor = Actor;
+					FractureContext.Transform = Actor->GetActorTransform();
+					FractureContext.OriginalPrimitiveComponent = GeometryCollectionComponent;
+					FractureContext.FracturedGeometryCollection = FracturedGeometryCollection;
+					FractureContext.SelectedBones = SelectedBones;
+
+					FractureContext.Bounds = FBox(ForceInit);
+					for (int32 BoneIndex : FractureContext.SelectedBones)
+					{
+						if (FractureContext.SelectedBones.Num() > 1 && RandomStream.FRand() > CommonSettings->ChanceToFracture)
 						{
-							ensure(TransformToGeometryIndex[Idx] > -1);
-							BoundsToBone.Add(Idx, BoundingBoxes[TransformToGeometryIndex[Idx]].TransformBy(Transforms[Idx]));
+							continue;
 						}
-					}
 
-					if (CommonSettings->bGroupFracture)
+						if (TransformToGeometryIndex[BoneIndex] > -1)
+						{
+							FractureContext.Bounds += BoundsToBone[BoneIndex];
+						}
+					}
+				}
+				else
+				{
+					for (int32 BoneIndex : SelectedBones)
 					{
+						if (SelectedBones.Num() > 1 && RandomStream.FRand() > CommonSettings->ChanceToFracture)
+						{
+							continue;
+						}
+
 						FractureContexts.AddDefaulted();
 						FFractureContext& FractureContext = FractureContexts.Last();
 						FractureContext.RandomSeed = FMath::Rand();
@@ -1145,49 +1177,10 @@ void FFractureEditorModeToolkit::GetFractureContexts(TArray<FFractureContext>& F
 						FractureContext.Transform = Actor->GetActorTransform();
 						FractureContext.OriginalPrimitiveComponent = GeometryCollectionComponent;
 						FractureContext.FracturedGeometryCollection = FracturedGeometryCollection;
-						FractureContext.SelectedBones = SelectedBones;
-
-						FractureContext.Bounds = FBox(ForceInit);
-						for (int32 BoneIndex : FractureContext.SelectedBones)
+						FractureContext.SelectedBones = { BoneIndex };
+						if(TransformToGeometryIndex[BoneIndex] > -1)
 						{
-							if (FractureContext.SelectedBones.Num() > 1 && RandomStream.FRand() > CommonSettings->ChanceToFracture)
-							{
-								continue;
-							}
-
-							if (TransformToGeometryIndex[BoneIndex] > -1)
-							{
-								FractureContext.Bounds += BoundsToBone[BoneIndex];
-							}
-						}
-					}
-					else
-					{
-						for (int32 BoneIndex : SelectedBones)
-						{
-							if (SelectedBones.Num() > 1 && RandomStream.FRand() > CommonSettings->ChanceToFracture)
-							{
-								continue;
-							}
-
-							FractureContexts.AddDefaulted();
-							FFractureContext& FractureContext = FractureContexts.Last();
-							FractureContext.RandomSeed = FMath::Rand();
-							if (CommonSettings->RandomSeed > -1)
-							{
-								// make sure it's unique for each context if it's specified.
-								FractureContext.RandomSeed = CommonSettings->RandomSeed + FractureContexts.Num();
-							}
-
-							FractureContext.OriginalActor = Actor;
-							FractureContext.Transform = Actor->GetActorTransform();
-							FractureContext.OriginalPrimitiveComponent = GeometryCollectionComponent;
-							FractureContext.FracturedGeometryCollection = FracturedGeometryCollection;
-							FractureContext.SelectedBones = { BoneIndex };
-							if(TransformToGeometryIndex[BoneIndex] > -1)
-							{
-								FractureContext.Bounds = BoundsToBone[BoneIndex];
-							}
+							FractureContext.Bounds = BoundsToBone[BoneIndex];
 						}
 					}
 				}
@@ -1330,15 +1323,9 @@ void FFractureEditorModeToolkit::GetSelectedGeometryCollectionComponents(TSet<UG
 
 	for (AActor* Actor : SelectedActors)
 	{
-		TArray<UActorComponent*> Components = Actor->GetComponentsByClass(UPrimitiveComponent::StaticClass());
-		for (UActorComponent* Component : Components)
-		{
-			UPrimitiveComponent* PrimitiveComponent = CastChecked<UPrimitiveComponent>(Component);
-			if (UGeometryCollectionComponent* GeometryCollectionComponent = Cast<UGeometryCollectionComponent>(PrimitiveComponent))
-			{
-				GeomCompSelection.Add(GeometryCollectionComponent);
-			}
-		}
+		TInlineComponentArray<UGeometryCollectionComponent*> GeometryCollectionComponents;
+		Actor->GetComponents(GeometryCollectionComponents);
+		GeomCompSelection.Append(GeometryCollectionComponents);
 	}
 }
 
@@ -1741,8 +1728,7 @@ bool FFractureEditorModeToolkit::IsGeometryCollectionSelected()
 		AActor* Actor = Cast<AActor>(*Iter);
 		if (Actor)
 		{
-			TArray<UActorComponent*> Components = Actor->GetComponentsByClass(UGeometryCollectionComponent::StaticClass());
-			if (Components.Num() > 0)
+			if (Actor->FindComponentByClass<UGeometryCollectionComponent>())
 			{
 				return true;
 			}
diff --git a/Engine/Plugins/Experimental/ChaosEditor/Source/FractureEditor/Private/FractureToolBrick.cpp b/Engine/Plugins/Experimental/ChaosEditor/Source/FractureEditor/Private/FractureToolBrick.cpp
--- a/Engine/Plugins/Experimental/ChaosEditor/Source/FractureEditor/Private/FractureToolBrick.cpp
+++ b/Engine/Plugins/Experimental/ChaosEditor/Source/FractureEditor/Private/FractureToolBrick.cpp
@@ -246,27 +246,25 @@ void UFractureToolBrick::PostEditChangeChainProperty(struct FPropertyChangedChai
 
 	for (AActor* Actor : SelectedActors)
 	{
-		TArray<UActorComponent*> Components = Actor->GetComponentsByClass(UPrimitiveComponent::StaticClass());
-		for (UActorComponent* Component : Components)
+		TInlineComponentArray<UPrimitiveComponent*> PrimitiveComponents;
+		Actor->GetComponents(PrimitiveComponents);
+		for (UPrimitiveComponent* PrimitiveComponent : PrimitiveComponents)
 		{
-			if (UPrimitiveComponent* PrimitiveComponent = CastChecked<UPrimitiveComponent>(Component))
+			FractureContext.OriginalActor = Actor;
+			FractureContext.OriginalPrimitiveComponent = PrimitiveComponent;
+			// 				FractureContext.Transform = PrimitiveComponent->GetComponentTransform();
+			FractureContext.Transform = Actor->GetTransform();
+			FVector Origin;
+			FVector BoxExtent;
+			Actor->GetActorBounds(false, Origin, BoxExtent);
+			if (LocalCommonSettings->bGroupFracture)
 			{
-				FractureContext.OriginalActor = Actor;
-				FractureContext.OriginalPrimitiveComponent = PrimitiveComponent;
-				// 				FractureContext.Transform = PrimitiveComponent->GetComponentTransform();
-				FractureContext.Transform = Actor->GetTransform();
-				FVector Origin;
-				FVector BoxExtent;
-				Actor->GetActorBounds(false, Origin, BoxExtent);
-				if (LocalCommonSettings->bGroupFracture)
-				{
-					FractureContext.Bounds += FBox::BuildAABB(Origin, BoxExtent);
-				}
-				else
-				{
-					FractureContext.Bounds = FBox::BuildAABB(Origin, BoxExtent);
-					GenerateBrickTransforms(FractureContext.Bounds);
-				}
+				FractureContext.Bounds += FBox::BuildAABB(Origin, BoxExtent);
+			}
+			else
+			{
+				FractureContext.Bounds = FBox::BuildAABB(Origin, BoxExtent);
+				GenerateBrickTransforms(FractureContext.Bounds);
 			}
 		}
 	}
diff --git a/Engine/Plugins/FX/Niagara/Source/Niagara/Private/NiagaraDataInterfaceSkeletalMesh.cpp b/Engine/Plugins/FX/Niagara/Source/Niagara/Private/NiagaraDataInterfaceSkeletalMesh.cpp
--- a/Engine/Plugins/FX/Niagara/Source/Niagara/Private/NiagaraDataInterfaceSkeletalMesh.cpp
+++ b/Engine/Plugins/FX/Niagara/Source/Niagara/Private/NiagaraDataInterfaceSkeletalMesh.cpp
@@ -920,8 +920,7 @@ USkeletalMesh* UNiagaraDataInterfaceSkeletalMesh::GetSkeletalMesh(UNiagaraCompon
 			}
 			else if (AActor* Actor = Cast<AActor>(UserParamObject))
 			{
-				TArray<UActorComponent*> SourceComps = Actor->GetComponentsByClass(USkeletalMeshComponent::StaticClass());
-				for (UActorComponent* ActorComp : SourceComps)
+				for (UActorComponent* ActorComp : Actor->GetComponents())
 				{
 					USkeletalMeshComponent* SourceComp = Cast<USkeletalMeshComponent>(ActorComp);
 					if (SourceComp)
@@ -994,8 +993,7 @@ USkeletalMesh* UNiagaraDataInterfaceSkeletalMesh::GetSkeletalMesh(UNiagaraCompon
 			}
 			else if (AActor* Owner = SimComp->GetAttachmentRootActor())
 			{
-				TArray<UActorComponent*> SourceComps = Owner->GetComponentsByClass(USkeletalMeshComponent::StaticClass());
-				for (UActorComponent* ActorComp : SourceComps)
+				for (UActorComponent* ActorComp : Owner->GetComponents())
 				{
 					USkeletalMeshComponent* SourceComp = Cast<USkeletalMeshComponent>(ActorComp);
 					if (SourceComp)
diff --git a/Engine/Plugins/FX/Niagara/Source/Niagara/Private/NiagaraDataInterfaceStaticMesh.cpp b/Engine/Plugins/FX/Niagara/Source/Niagara/Private/NiagaraDataInterfaceStaticMesh.cpp
--- a/Engine/Plugins/FX/Niagara/Source/Niagara/Private/NiagaraDataInterfaceStaticMesh.cpp
+++ b/Engine/Plugins/FX/Niagara/Source/Niagara/Private/NiagaraDataInterfaceStaticMesh.cpp
@@ -195,8 +195,7 @@ bool FNDIStaticMesh_InstanceData::Init(UNiagaraDataInterfaceStaticMesh* Interfac
 			}
 			else if (AActor* Owner = SimComp->GetAttachmentRootActor())
 			{
-				TArray<UActorComponent*> SourceComps = Owner->GetComponentsByClass(UStaticMeshComponent::StaticClass());
-				for (UActorComponent* ActorComp : SourceComps)
+				for (UActorComponent* ActorComp : Owner->GetComponents())
 				{
 					UStaticMeshComponent* SourceComp = Cast<UStaticMeshComponent>(ActorComp);
 					if (SourceComp)
diff --git a/Engine/Plugins/MovieScene/MatineeToLevelSequence/Source/MatineeToLevelSequence/Private/MatineeToLevelSequenceModule.cpp b/Engine/Plugins/MovieScene/MatineeToLevelSequence/Source/MatineeToLevelSequence/Private/MatineeToLevelSequenceModule.cpp
--- a/Engine/Plugins/MovieScene/MatineeToLevelSequence/Source/MatineeToLevelSequence/Private/MatineeToLevelSequenceModule.cpp
+++ b/Engine/Plugins/MovieScene/MatineeToLevelSequence/Source/MatineeToLevelSequence/Private/MatineeToLevelSequenceModule.cpp
@@ -343,7 +343,7 @@ protected:
 		{
 			FolderName = TEXT("Lights");
 		}
-		else if (Actor.Get()->GetComponentsByClass(UParticleSystemComponent::StaticClass()).Num())
+		else if (Actor.Get()->FindComponentByClass<UParticleSystemComponent>())
 		{
 			FolderName = TEXT("Particles");
 		}
diff --git a/Engine/Source/Editor/MeshPaintMode/Private/PaintModePainter.cpp b/Engine/Source/Editor/MeshPaintMode/Private/PaintModePainter.cpp
--- a/Engine/Source/Editor/MeshPaintMode/Private/PaintModePainter.cpp
+++ b/Engine/Source/Editor/MeshPaintMode/Private/PaintModePainter.cpp
@@ -2361,10 +2361,11 @@ TArray<ComponentClass*> FPaintModePainter::GetSelectedComponents() const
 			AActor* SelectedActor = Cast<AActor>(ActorSelection->GetSelectedObject(SelectionIndex));
 			if (SelectedActor)
 			{
-				TArray<UActorComponent*> ActorComponents = SelectedActor->GetComponentsByClass(ComponentClass::StaticClass());
-				for (UActorComponent* Component : ActorComponents)
+				TInlineComponentArray<ComponentClass*> ActorComponents;
+				SelectedActor->GetComponents(ActorComponents);
+				for (ComponentClass* Component : ActorComponents)
 				{
-					Components.AddUnique(CastChecked<ComponentClass>(Component));
+					Components.AddUnique(Component);
 				}
 			}
 		}
diff --git a/Engine/Source/Runtime/Experimental/GeometryCollectionEngine/Private/GeometryCollection/GeometryCollectionComponent.cpp b/Engine/Source/Runtime/Experimental/GeometryCollectionEngine/Private/GeometryCollection/GeometryCollectionComponent.cpp
--- a/Engine/Source/Runtime/Experimental/GeometryCollectionEngine/Private/GeometryCollection/GeometryCollectionComponent.cpp
+++ b/Engine/Source/Runtime/Experimental/GeometryCollectionEngine/Private/GeometryCollection/GeometryCollectionComponent.cpp
@@ -2144,8 +2144,9 @@ void UGeometryCollectionComponent::SwitchRenderModels(const AActor* Actor)
 		return;
 	}
 
-	TArray<UActorComponent*> PrimitiveComponents = Actor->GetComponentsByClass(UPrimitiveComponent::StaticClass());
-	for (UActorComponent* PrimitiveComponent : PrimitiveComponents)
+	TInlineComponentArray<UPrimitiveComponent*> PrimitiveComponents;
+	Actor->GetComponents(PrimitiveComponents);
+	for (UPrimitiveComponent* PrimitiveComponent : PrimitiveComponents)
 	{
 		bool ValidComponent = false;
 
@@ -2164,8 +2165,8 @@ void UGeometryCollectionComponent::SwitchRenderModels(const AActor* Actor)
 		}
 	}
 
-	TArray<UChildActorComponent*> ChildActorComponents;
-	Actor->GetComponents<UChildActorComponent>(ChildActorComponents);
+	TInlineComponentArray<UChildActorComponent*> ChildActorComponents;
+	Actor->GetComponents(ChildActorComponents);
 	for (UChildActorComponent* ChildComponent : ChildActorComponents)
 	{
 		AActor* ChildActor = ChildComponent->GetChildActor();
diff --git a/Engine/Source/Runtime/Experimental/InteractiveToolsFramework/Public/ToolBuilderUtil.h b/Engine/Source/Runtime/Experimental/InteractiveToolsFramework/Public/ToolBuilderUtil.h
--- a/Engine/Source/Runtime/Experimental/InteractiveToolsFramework/Public/ToolBuilderUtil.h
+++ b/Engine/Source/Runtime/Experimental/InteractiveToolsFramework/Public/ToolBuilderUtil.h
@@ -68,9 +68,13 @@ int ToolBuilderUtil::CountSelectedComponentsOfType(const FToolBuilderState& Inpu
 		{
 			if (AActor* Actor = Cast<AActor>(*Iter))
 			{
-				// [RMS] is there a more efficient way to do this?
-				TArray<UActorComponent*> TypedComponents = Actor->GetComponentsByClass(ComponentType::StaticClass());
-				nTypedComponents += TypedComponents.Num();
+				for (UActorComponent* Component : Actor->GetComponents())
+				{
+					if (Cast<ComponentType>(Component))
+					{
+						++nTypedComponents;
+					}
+				}
 			}
 		}
 	}
@@ -102,11 +106,9 @@ ComponentType* ToolBuilderUtil::FindFirstComponentOfType(const FToolBuilderState
 		{
 			if (AActor* Actor = Cast<AActor>(*Iter))
 			{
-				// [RMS] is there a more efficient way to do this?
-				TArray<UActorComponent*> TypedComponents = Actor->GetComponentsByClass(ComponentType::StaticClass());
-				if (TypedComponents.Num() > 0)
+				if (ComponentType* TypedComponent = Actor->FindComponentByClass<ComponentType>())
 				{
-					return (ComponentType*)TypedComponents[0];
+					return TypedComponent;
 				}
 			}
 		}
diff --git a/Engine/Source/Runtime/LevelSequence/Private/LevelSequenceActorSpawner.cpp b/Engine/Source/Runtime/LevelSequence/Private/LevelSequenceActorSpawner.cpp
--- a/Engine/Source/Runtime/LevelSequence/Private/LevelSequenceActorSpawner.cpp
+++ b/Engine/Source/Runtime/LevelSequence/Private/LevelSequenceActorSpawner.cpp
@@ -120,10 +120,12 @@ UObject* FLevelSequenceActorSpawner::SpawnObject(FMovieSceneSpawnable& Spawnable
 
 	{
 		// Disable all particle components so that they don't auto fire as soon as the actor is spawned. The particles should be triggered through the particle track.
-		TArray<UActorComponent*> ParticleComponents = ObjectTemplate->GetComponentsByClass(UParticleSystemComponent::StaticClass());
-		for (int32 ComponentIdx = 0; ComponentIdx < ParticleComponents.Num(); ++ComponentIdx)
+		for (UActorComponent* Component : ObjectTemplate->GetComponents())
 		{
-			ParticleComponents[ComponentIdx]->bAutoActivate = false;
+			if (UParticleSystemComponent* ParticleComponent = Cast<UParticleSystemComponent>(Component))
+			{
+				Component->bAutoActivate = false;
+			}
 		}
 	}
 
