Mitigated crash when going into PIE/SIE with a selected geometry collection with explode vectors active. Underlying tools cause implicits to be destroyed and recreated while they are pending initialisation on other threads.

diff --git a/Engine/Plugins/Experimental/ChaosEditor/Source/FractureEditor/Private/FractureEditorModeToolkit.cpp b/Engine/Plugins/Experimental/ChaosEditor/Source/FractureEditor/Private/FractureEditorModeToolkit.cpp
--- a/Engine/Plugins/Experimental/ChaosEditor/Source/FractureEditor/Private/FractureEditorModeToolkit.cpp
+++ b/Engine/Plugins/Experimental/ChaosEditor/Source/FractureEditor/Private/FractureEditorModeToolkit.cpp
@@ -1605,6 +1605,14 @@ bool GetValidGeoCenter(const TManagedArray<int32>& TransformToGeometryIndex, con
 
 void FFractureEditorModeToolkit::UpdateExplodedVectors(UGeometryCollectionComponent* GeometryCollectionComponent) const
 {
+#if WITH_EDITOR
+	// If we're running PIE or SIE when this happens we should ignore the rebuild as the implicits will be in use.
+	if(GEditor->bIsSimulatingInEditor || GEditor->GetPIEWorldContext() != nullptr)
+	{
+		return;
+	}
+#endif
+
 	TSharedPtr<FGeometryCollection, ESPMode::ThreadSafe> GeometryCollectionPtr = GeometryCollectionComponent->GetRestCollection()->GetGeometryCollection();
 	const FGeometryCollection* OutGeometryCollectionConst = GeometryCollectionPtr.Get();
 
