Fix crash in fracture tool stats gathering after deleting geometry collection assets that are set on a selected actor in the editor viewport

diff --git a/Engine/Plugins/Experimental/ChaosEditor/Source/FractureEditor/Private/FractureEditorModeToolkit.cpp b/Engine/Plugins/Experimental/ChaosEditor/Source/FractureEditor/Private/FractureEditorModeToolkit.cpp
--- a/Engine/Plugins/Experimental/ChaosEditor/Source/FractureEditor/Private/FractureEditorModeToolkit.cpp
+++ b/Engine/Plugins/Experimental/ChaosEditor/Source/FractureEditor/Private/FractureEditorModeToolkit.cpp
@@ -1782,11 +1782,15 @@ FText FFractureEditorModeToolkit::GetStatisticsSummary() const
 			if (AGeometryCollectionActor* Actor = Cast<AGeometryCollectionActor>(*Iter))
 			{
 				const UGeometryCollection* RestCollection = Actor->GetGeometryCollectionComponent()->GetRestCollection();
-				const FGeometryCollection* GeometryCollection = RestCollection->GetGeometryCollection().Get();
 
-				if (GeometryCollection != nullptr)
+				if(RestCollection)
 				{
-					GeometryCollectionArray.Add(GeometryCollection);
+					const FGeometryCollection* GeometryCollection = RestCollection->GetGeometryCollection().Get();
+
+					if(GeometryCollection != nullptr)
+					{
+						GeometryCollectionArray.Add(GeometryCollection);
+					}
 				}
 			}
 		}
@@ -1810,28 +1814,31 @@ FText FFractureEditorModeToolkit::GetStatisticsSummary() const
 
 			Buffer += FString::Printf(TEXT("Sum of the selected Geometry Collections\n\n"));
 
-			const TManagedArray<int32>& Levels = GeometryCollection->GetAttribute<int32>("Level", FGeometryCollection::TransformGroup);
-
-			TArray<int32> LevelTransforms;
-			for (int32 Element = 0, NumElement = Levels.Num(); Element < NumElement; ++Element)
+			if(GeometryCollection->HasAttribute("Level", FGeometryCollection::TransformGroup))
 			{
-				const int32 NodeLevel = Levels[Element];
-				while (LevelTransforms.Num() <= NodeLevel)
+				const TManagedArray<int32>& Levels = GeometryCollection->GetAttribute<int32>("Level", FGeometryCollection::TransformGroup);
+
+				TArray<int32> LevelTransforms;
+				for(int32 Element = 0, NumElement = Levels.Num(); Element < NumElement; ++Element)
 				{
-					LevelTransforms.SetNum(NodeLevel + 1);
-					LevelTransforms[NodeLevel] = 0;
+					const int32 NodeLevel = Levels[Element];
+					while(LevelTransforms.Num() <= NodeLevel)
+					{
+						LevelTransforms.SetNum(NodeLevel + 1);
+						LevelTransforms[NodeLevel] = 0;
+					}
+					++LevelTransforms[NodeLevel];
 				}
-				++LevelTransforms[NodeLevel];
-			}
 
-			for (int32 Level = 0; Level < LevelTransforms.Num(); ++Level)
-			{
-				LevelTransformsAll[Level] += LevelTransforms[Level];
-			}
+				for(int32 Level = 0; Level < LevelTransforms.Num(); ++Level)
+				{
+					LevelTransformsAll[Level] += LevelTransforms[Level];
+				}
 
-			if (LevelTransforms.Num() > LevelMax)
-			{
-				LevelMax = LevelTransforms.Num();
+				if(LevelTransforms.Num() > LevelMax)
+				{
+					LevelMax = LevelTransforms.Num();
+				}
 			}
 		}
 
