Fracture Updated to Use new Modes Toolbars

diff --git a/Engine/Plugins/Experimental/ChaosEditor/Source/FractureEditor/Private/FractureEditorCommands.cpp b/Engine/Plugins/Experimental/ChaosEditor/Source/FractureEditor/Private/FractureEditorCommands.cpp
--- a/Engine/Plugins/Experimental/ChaosEditor/Source/FractureEditor/Private/FractureEditorCommands.cpp
+++ b/Engine/Plugins/Experimental/ChaosEditor/Source/FractureEditor/Private/FractureEditorCommands.cpp
@@ -17,15 +17,15 @@ void FFractureEditorCommands::RegisterCommands()
 {
 
 	// Selection Commands
-	UI_COMMAND(SelectAll, "SelectAll", "Selects all Bones in the GeometryCollection", EUserInterfaceActionType::Button, FInputChord(EModifierKey::Control | EModifierKey::Shift, EKeys::A));
-	UI_COMMAND(SelectNone, "SelectNone", "Deselects all Bones in the GeometryCollection", EUserInterfaceActionType::Button, FInputChord(EModifierKey::Control | EModifierKey::Shift, EKeys::D));
-	UI_COMMAND(SelectNeighbors, "SelectNeighbors", "Select all bones adjacent to the currently selected bones", EUserInterfaceActionType::Button, FInputChord() );
-	UI_COMMAND(SelectSiblings, "SelectSiblings", "Select all bones at the same levels as the currently selected bones", EUserInterfaceActionType::Button, FInputChord() );
-	UI_COMMAND(SelectAllInCluster, "SelectAllInCluster", "Select all bones with the same parent as selected bones", EUserInterfaceActionType::Button, FInputChord() );
-	UI_COMMAND(SelectInvert, "SelectInvert", "Invert the selected bones", EUserInterfaceActionType::Button, FInputChord() );
+	UI_COMMAND(SelectAll, "Select All", "Selects all Bones in the GeometryCollection", EUserInterfaceActionType::Button, FInputChord(EModifierKey::Control | EModifierKey::Shift, EKeys::A));
+	UI_COMMAND(SelectNone, "Deselect", "Deselects all Bones in the GeometryCollection", EUserInterfaceActionType::Button, FInputChord(EModifierKey::Control | EModifierKey::Shift, EKeys::D));
+	UI_COMMAND(SelectNeighbors, "Neighbors", "Select all bones adjacent to the currently selected bones", EUserInterfaceActionType::Button, FInputChord() );
+	UI_COMMAND(SelectSiblings, "Level", "Select all bones at the same levels as the currently selected bones", EUserInterfaceActionType::Button, FInputChord() );
+	UI_COMMAND(SelectAllInCluster, "Cluster", "Select all bones with the same parent as selected bones", EUserInterfaceActionType::Button, FInputChord() );
+	UI_COMMAND(SelectInvert, "Invert", "Invert the selected bones", EUserInterfaceActionType::Button, FInputChord() );
 
 	// View settings
-	UI_COMMAND(ToggleShowBoneColors, "ShowBoneColors", "Toggle Show Bone Colors", EUserInterfaceActionType::ToggleButton, FInputChord(EModifierKey::Shift, EKeys::B) );
+	UI_COMMAND(ToggleShowBoneColors, "Colors", "Toggle Show Bone Colors", EUserInterfaceActionType::ToggleButton, FInputChord(EModifierKey::Shift, EKeys::B) );
 	UI_COMMAND(ViewUpOneLevel, "ViewUpOneLevel", "View Up One Level", EUserInterfaceActionType::Button, FInputChord(EModifierKey::Shift, EKeys::W) );
 	UI_COMMAND(ViewDownOneLevel, "ViewDownOneLevel", "View Down One Level", EUserInterfaceActionType::Button, FInputChord(EModifierKey::Shift, EKeys::S) );
 	UI_COMMAND(ExplodeMore, "ExplodeMore", "Explode 10% More", EUserInterfaceActionType::Button, FInputChord(EModifierKey::Shift, EKeys::E) );
@@ -34,13 +34,13 @@ void FFractureEditorCommands::RegisterCommands()
 	// Cluster Operations
 	UI_COMMAND(Flatten, "Flatten", "Flattens all bones to level 1", EUserInterfaceActionType::Button, FInputChord() );
 	UI_COMMAND(FlattenToLevel, "FlattenToLevel", "Flatten Only up to the Current View Level", EUserInterfaceActionType::Button, FInputChord() );
-	UI_COMMAND(AutoCluster, "AutoCluster", "Auto Cluster", EUserInterfaceActionType::Button, FInputChord() );
+	UI_COMMAND(AutoCluster, "Auto", "Auto Cluster", EUserInterfaceActionType::Button, FInputChord() );
 	UI_COMMAND(Cluster, "Cluster", "Clusters selected bones under a new parent", EUserInterfaceActionType::Button, FInputChord() );
 	UI_COMMAND(Uncluster, "Uncluster", "Uncluster", EUserInterfaceActionType::Button, FInputChord() );
 	UI_COMMAND(Merge, "Merge", "Merge", EUserInterfaceActionType::Button, FInputChord() );
-	UI_COMMAND(MoveUp, "MoveUp", "Move Bones Up One Level", EUserInterfaceActionType::Button, FInputChord() );
+	UI_COMMAND(MoveUp, "Level Up", "Move Bones Up One Level", EUserInterfaceActionType::Button, FInputChord() );
 
-	UI_COMMAND(GenerateAsset, "GenerateAsset", "Generate Geometry Collection Asset from static meshes contained in selected actors", EUserInterfaceActionType::Button, FInputChord() );
+	UI_COMMAND(GenerateAsset, "New", "Generate Geometry Collection Asset from static meshes contained in selected actors", EUserInterfaceActionType::Button, FInputChord() );
 
 
 	for (TObjectIterator<UClass> ClassIterator; ClassIterator; ++ClassIterator)
diff --git a/Engine/Plugins/Experimental/ChaosEditor/Source/FractureEditor/Private/FractureEditorMode.cpp b/Engine/Plugins/Experimental/ChaosEditor/Source/FractureEditor/Private/FractureEditorMode.cpp
--- a/Engine/Plugins/Experimental/ChaosEditor/Source/FractureEditor/Private/FractureEditorMode.cpp
+++ b/Engine/Plugins/Experimental/ChaosEditor/Source/FractureEditor/Private/FractureEditorMode.cpp
@@ -118,10 +118,13 @@ void FFractureEditorMode::Render(const FSceneView* View, FViewport* Viewport, FP
 
 	FFractureEditorModeToolkit* FractureToolkit = (FFractureEditorModeToolkit*)Toolkit.Get();
  
-	UFractureTool* FractureTool = FractureToolkit->GetActiveTool();
+	if (UFractureTool* FractureTool = FractureToolkit->GetActiveTool())
+	{
+		auto Settings = FractureTool->GetSettingsObjects();
+		FractureTool->Render(View, Viewport, PDI);
+	}
+
 
-	auto Settings = FractureTool->GetSettingsObjects();
-	FractureTool->Render(View, Viewport, PDI);
 }
 
 bool FFractureEditorMode::UsesToolkits() const
diff --git a/Engine/Plugins/Experimental/ChaosEditor/Source/FractureEditor/Private/FractureEditorModeToolkit.cpp b/Engine/Plugins/Experimental/ChaosEditor/Source/FractureEditor/Private/FractureEditorModeToolkit.cpp
--- a/Engine/Plugins/Experimental/ChaosEditor/Source/FractureEditor/Private/FractureEditorModeToolkit.cpp
+++ b/Engine/Plugins/Experimental/ChaosEditor/Source/FractureEditor/Private/FractureEditorModeToolkit.cpp
@@ -13,9 +13,11 @@
 #include "PropertyEditorModule.h"
 #include "Widgets/Input/SButton.h"
 #include "Widgets/Layout/SBox.h"
+#include "Widgets/Layout/SSpacer.h"
 #include "Widgets/Layout/SExpandableArea.h"
 #include "Widgets/Layout/SScrollBox.h"
 #include "Widgets/Layout/SSplitter.h"
+#include "Widgets/Layout/SWidgetSwitcher.h"
 #include "Widgets/Input/SSlider.h"
 #include "Widgets/Input/SCheckBox.h"
 #include "Widgets/Input/SComboBox.h"
@@ -83,8 +85,6 @@ FFractureEditorModeToolkit::FFractureEditorModeToolkit()
 	, AutoClusterMode(EFractureAutoClusterMode::BoundingBox)
 	, AutoClusterSiteCount(10)
 	, ActiveTool(nullptr)
-	, bFractureGroupExpanded(true)
-	, bOutlinerGroupExpanded(true)
 {
 }
 
@@ -93,38 +93,6 @@ void FFractureEditorModeToolkit::Init(const TSharedPtr<IToolkitHost>& InitToolki
 	FFractureEditorModule& FractureModule = FModuleManager::GetModuleChecked<FFractureEditorModule>("FractureEditor");
 
 	const FFractureEditorCommands& Commands = FFractureEditorCommands::Get();
-	FToolBarBuilder SelectionBuilder(GetToolkitCommands(), FMultiBoxCustomization::None, nullptr, Orient_Horizontal, true);
-	SelectionBuilder.AddToolBarButton(Commands.GenerateAsset);
-	SelectionBuilder.AddSeparator();
-	SelectionBuilder.AddToolBarButton(Commands.SelectAll);
-	SelectionBuilder.AddToolBarButton(Commands.SelectNone);
-	SelectionBuilder.AddToolBarButton(Commands.SelectNeighbors);
-	SelectionBuilder.AddToolBarButton(Commands.SelectSiblings);
-	SelectionBuilder.AddToolBarButton(Commands.SelectAllInCluster);
-	SelectionBuilder.AddToolBarButton(Commands.SelectInvert);
-
-	FToolBarBuilder ClusterToolBuilder(GetToolkitCommands(), FMultiBoxCustomization::None, nullptr, Orient_Horizontal, true);
-	ClusterToolBuilder.AddToolBarButton(Commands.AutoCluster);
-	ClusterToolBuilder.AddToolBarButton(Commands.Flatten);
-	//ClusterToolBuilder.AddToolBarButton(Commands.FlattenToLevel);
-	ClusterToolBuilder.AddToolBarButton(Commands.Cluster);
-	ClusterToolBuilder.AddToolBarButton(Commands.Uncluster);
-// 	ClusterToolBuilder.AddToolBarButton(Commands.Merge);
-	ClusterToolBuilder.AddToolBarButton(Commands.MoveUp);
-
-	// FToolBarBuilder ViewToolBuilder(GetToolkitCommands(), FMultiBoxCustomization::None, nullptr, Orient_Horizontal, true);
-	// ViewToolBuilder.AddToolBarButton(Commands.ViewUpOneLevel);
-	// ViewToolBuilder.AddToolBarButton(Commands.ViewDownOneLevel);
-
-	TSharedPtr<FExtender> Extender = FractureModule.GetToolBarExtensibilityManager()->GetAllExtenders();
-	FToolBarBuilder FractureToolBuilder(GetToolkitCommands(), FMultiBoxCustomization::None, Extender, Orient_Horizontal, true);// , true);
-	FractureToolBuilder.AddToolBarButton(Commands.Uniform);
-	FractureToolBuilder.AddToolBarButton(Commands.Clustered);
-	FractureToolBuilder.AddToolBarButton(Commands.Radial);
-	FractureToolBuilder.AddToolBarButton(Commands.Planar);
-	FractureToolBuilder.AddToolBarButton(Commands.Slice);
-	FractureToolBuilder.AddToolBarButton(Commands.Brick);
-// 	FractureToolBuilder.AddToolBarButton(Commands.Texture);
 
 	FPropertyEditorModule& EditModule = FModuleManager::Get().GetModuleChecked<FPropertyEditorModule>("PropertyEditor");
 
@@ -151,7 +119,86 @@ void FFractureEditorModeToolkit::Init(const TSharedPtr<IToolkitHost>& InitToolki
 	float Padding = 4.0f;
 	float MorePadding = 10.0f;
 
-	SAssignNew(ToolkitWidget, SBorder)
+	SAssignNew(ExplodedViewWidget, SSpinBox<int32>)
+	.Style(&FFractureEditorStyle::Get().GetWidgetStyle<FSpinBoxStyle>("FractureEditor.SpinBox"))
+	.PreventThrottling(true)
+	.Value_Lambda([=]() -> int32 { return (int32) (ExplodeAmount * 100.0f); })
+	.OnValueChanged_Lambda( [=](int32 NewValue) { this->OnSetExplodedViewValue( (float)NewValue / 100.0f); } )
+	.MinValue(0)
+	.MaxValue(100)
+	.Delta(1)
+	.Font(FCoreStyle::GetDefaultFontStyle("Regular", 14))
+	.MinDesiredWidth(36.f)
+	.Justification(ETextJustify::Center)
+	.ToolTipText(LOCTEXT("FractureEditor.Exploded_Tooltip", "How much to seperate the drawing of the bones to aid in setup.  Does not effect simulation"))
+	;
+
+	SAssignNew(LevelViewWidget, SComboButton)
+	.ContentPadding(0)
+	.ButtonStyle(FEditorStyle::Get(), "Toolbar.Button")
+	.ForegroundColor(FEditorStyle::Get().GetSlateColor("ToolBar.SToolBarComboButtonBlock.ComboButton.Color"))
+	.ButtonContent()
+	[
+		SNew(STextBlock)
+		.Text_Lambda( [=]() -> FText {
+			if (FractureLevel < 0)
+			{
+				return LOCTEXT("FractureViewAllLevels", "All");
+			}
+			else if (FractureLevel == 0)
+			{
+				return LOCTEXT("FractureViewRootLevel", "Root");
+			}
+
+			return FText::Format(NSLOCTEXT("FractureEditor", "CurrentLevel", "{0}"), FText::AsNumber(FractureLevel));
+
+		})
+		.Font(FCoreStyle::GetDefaultFontStyle("Regular", 12))
+		.ToolTipText(LOCTEXT("FractureEditor.Level_Tooltip", "Set the currently view level of the geometry collection"))
+
+	];
+
+	SAssignNew(ShowBoneColorsWidget, SComboButton)
+	.ContentPadding(0)
+	.ButtonStyle(FEditorStyle::Get(), "ToolBar.Button")
+	.ForegroundColor(FEditorStyle::Get().GetSlateColor("ToolBar.SToolBarComboButtonBlock.ComboButton.Color"))
+	.OnGetMenuContent(this, &FFractureEditorModeToolkit::GetViewMenuContent)
+	.ButtonContent()
+	[
+		SNew(SImage)
+		.Image(FFractureEditorStyle::Get().GetBrush("FractureEditor.Visibility"))
+		.ToolTipText(LOCTEXT("FractureEditor.Visibility_Tooltip", "Toggle showing the bone colours of the geometry collection"))
+	];
+
+	TSharedRef<SExpandableArea> OutlinerExpander = SNew(SExpandableArea)
+	.AreaTitle(FText(LOCTEXT("Outliner", "Outliner")))
+	.HeaderPadding(FMargin(2.0, 2.0))
+	.Padding(FMargin(MorePadding))
+	.BorderImage(FEditorStyle::Get().GetBrush("DetailsView.CategoryTop"))
+	.BorderBackgroundColor( FLinearColor( .6,.6,.6, 1.0f ) )
+	.BodyBorderBackgroundColor (FLinearColor( 1.0, 0.0, 0.0))
+	.AreaTitleFont(FEditorStyle::Get().GetFontStyle("DetailsView.CategoryFontStyle"))
+	.BodyContent()
+	[
+		SAssignNew(OutlinerView, SGeometryCollectionOutliner )
+		.OnBoneSelectionChanged(this, &FFractureEditorModeToolkit::OnOutlinerBoneSelectionChanged)
+	];
+
+	TSharedRef<SExpandableArea> StatisticsExpander = SNew(SExpandableArea)
+	.AreaTitle(FText(LOCTEXT("LevelStatistics", "Level Statistics")))
+	.HeaderPadding(FMargin(2.0, 2.0))
+	.Padding(FMargin(MorePadding))
+	.BorderImage(FEditorStyle::Get().GetBrush("DetailsView.CategoryTop"))
+	.BorderBackgroundColor( FLinearColor( .6,.6,.6, 1.0f ) )
+	.BodyBorderBackgroundColor (FLinearColor( 1.0, 0.0, 0.0))
+	.AreaTitleFont(FEditorStyle::Get().GetFontStyle("DetailsView.CategoryFontStyle"))
+	.BodyContent()
+	[
+		SNew(STextBlock)
+		.Text(this, &FFractureEditorModeToolkit::GetStatisticsSummary)
+	];
+
+	SAssignNew(ToolkitWidget, SBox)
 	.Padding(8)
 	[
 		SNew(SVerticalBox)
@@ -167,245 +214,35 @@ void FFractureEditorModeToolkit::Init(const TSharedPtr<IToolkitHost>& InitToolki
 		]
 
 		+SVerticalBox::Slot()
-		.AutoHeight()
 		[
-			SNew(SExpandableArea)
-			.AreaTitle(FText(LOCTEXT("Select", "Select")))
-			.HeaderPadding(FMargin(2.0, 2.0))
-			.BorderImage(FEditorStyle::Get().GetBrush("DetailsView.CategoryTop"))
-			.BorderBackgroundColor( FLinearColor( .6,.6,.6, 1.0f ) )
-			.BodyBorderBackgroundColor (FLinearColor( 0.8, 0.8, 0.8))
-			.AreaTitleFont(FEditorStyle::Get().GetFontStyle("DetailsView.CategoryFontStyle"))
-			.Padding(FMargin(MorePadding))
-			.BodyContent()
-			[
-
-				SNew(SVerticalBox)
-
-				+SVerticalBox::Slot()
-				.HAlign(HAlign_Center)
-				[
-					SelectionBuilder.MakeWidget()
-				]
-			]
-		]
+			SNew(SWidgetSwitcher)
+			.WidgetIndex_Lambda( [this] () -> uint32 { return (GetActiveTool() != nullptr ? 1 : 0); })
 
-		+SVerticalBox::Slot()
-		.AutoHeight()
-		[
-			SNew(SExpandableArea)
-			.AreaTitle(FText(LOCTEXT("ViewSettings", "View")))
-			.HeaderPadding(FMargin(2.0, 2.0))
-			.Padding(FMargin(MorePadding))
-			.BorderImage(FEditorStyle::Get().GetBrush("DetailsView.CategoryTop"))
-			.BorderBackgroundColor( FLinearColor( .6,.6,.6, 1.0f ) )
-			.BodyBorderBackgroundColor (FLinearColor( 1.0, 0.0, 0.0))
-			.AreaTitleFont(FEditorStyle::Get().GetFontStyle("DetailsView.CategoryFontStyle"))
-			.BodyContent()
+			+ SWidgetSwitcher::Slot()
 			[
-				SNew(SHorizontalBox)
-
-				+SHorizontalBox::Slot()
-				.FillWidth(.5)
-
-				+ SHorizontalBox::Slot()
-				.HAlign(HAlign_Center)
-				.VAlign(VAlign_Bottom)
-				.AutoWidth()
-				[
-					SNew(SImage)
-					.Image(FFractureEditorStyle::Get().GetBrush("FractureEditor.Exploded"))
-				]
-
-				+ SHorizontalBox::Slot()
-				.AutoWidth()
-				[
-					SNew(SSpinBox<int32>)
-					.Style(&FFractureEditorStyle::Get().GetWidgetStyle<FSpinBoxStyle>("FractureEditor.SpinBox"))
-					.PreventThrottling(true)
-					.Value_Lambda([=]() -> int32 { return (int32) (ExplodeAmount * 100.0f); })
-					.OnValueChanged_Lambda( [=](int32 NewValue) { this->OnSetExplodedViewValue( (float)NewValue / 100.0f); } )
-					.MinValue(0)
-					.MaxValue(100)
-					.Delta(1)
-					.Font(FCoreStyle::GetDefaultFontStyle("Regular", 18))
-					.ToolTipText(LOCTEXT("FractureEditor.Exploded_Tooltip", "How much to seperate the drawing of the bones to aid in setup.  Does not effect simulation"))
-				]
-
-				+SHorizontalBox::Slot()
-				.HAlign(HAlign_Left)
-				.VAlign(VAlign_Center)
-				.AutoWidth()
-				[
-					SNew(STextBlock)
-					.Text(FText(LOCTEXT("FractureExplodedPercentage", "%")))
-					.Font(FCoreStyle::GetDefaultFontStyle("Regular", 16))
-				]
-
-				+SHorizontalBox::Slot()
-				.FillWidth(1)
-
-
-				+SHorizontalBox::Slot()
-				.HAlign(HAlign_Left)
-				.VAlign(VAlign_Bottom)
-				.AutoWidth()
+				SNew(SSplitter)
+				.Orientation( Orient_Vertical )
+
+				+SSplitter::Slot()
+				.SizeRule( TAttribute<SSplitter::ESizeRule>::Create( [this, OutlinerExpander] () { 
+					return OutlinerExpander->IsExpanded() ? SSplitter::ESizeRule::FractionOfParent : SSplitter::ESizeRule::SizeToContent; 
+				} ) )
+				.Value(1.f)
 				[
-					SNew(STextBlock)
-					.Text_Lambda( [=]() -> FText { return FText::Format(NSLOCTEXT("FractureEditor", "TotalLevelCount", "{0}"), FText::AsNumber(GetLevelCount())); } )
-					.Font(FCoreStyle::GetDefaultFontStyle("Regular", 10))
+					OutlinerExpander
 				]
 
-				+SHorizontalBox::Slot()
-				.HAlign(HAlign_Center)
-				.VAlign(VAlign_Bottom)
-				.AutoWidth()
-				.Padding(Padding * 0.5, 0.0, 0.0, Padding * 0.5)
+				+SSplitter::Slot()
+				.SizeRule( TAttribute<SSplitter::ESizeRule>::Create( [this, StatisticsExpander] () { 
+					return StatisticsExpander->IsExpanded() ? SSplitter::ESizeRule::FractionOfParent : SSplitter::ESizeRule::SizeToContent; 
+				} ) )
+				.Value(0.25f)
 				[
-					SNew(SImage)
-					.Image(FFractureEditorStyle::Get().GetBrush("FractureEditor.Levels"))
-				]
-
-				+SHorizontalBox::Slot()
-				.AutoWidth()
-				[
-					SNew(SComboButton)
-					.ContentPadding(0)
-					.ButtonStyle(FEditorStyle::Get(), "ToolBar.Button")
-					.ForegroundColor(FEditorStyle::Get().GetSlateColor("ToolBar.SToolBarComboButtonBlock.ComboButton.Color"))
-					.ButtonContent()
-					[
-						SNew(STextBlock)
-						.Text_Lambda( [=]() -> FText {
-							if (FractureLevel < 0)
-							{
-								return LOCTEXT("FractureViewAllLevels", "All");
-							}
-							else if (FractureLevel == 0)
-							{
-								return LOCTEXT("FractureViewRootLevel", "Root");
-							}
-
-							return FText::Format(NSLOCTEXT("FractureEditor", "CurrentLevel", "{0}"), FText::AsNumber(FractureLevel));
-
-						})
-						.Font(FCoreStyle::GetDefaultFontStyle("Regular", 18))
-						.ToolTipText(LOCTEXT("FractureEditor.Level_Tooltip", "Set the currently view level of the geometry collection"))
-
-					]
-					.OnGetMenuContent(this, &FFractureEditorModeToolkit::GetLevelViewMenuContent)
-				]
-
-
-				+SHorizontalBox::Slot()
-				.FillWidth(1.25)
-
-				+SHorizontalBox::Slot()
-				.AutoWidth()
-				[
-					SNew(SComboButton)
-					.ContentPadding(0)
-					.ButtonStyle(FEditorStyle::Get(), "ToolBar.Button")
-					.ForegroundColor(FEditorStyle::Get().GetSlateColor("ToolBar.SToolBarComboButtonBlock.ComboButton.Color"))
-					.ButtonContent()
-					[
-						SNew(SImage)
-						.Image(FFractureEditorStyle::Get().GetBrush("FractureEditor.Visibility"))
-						.ToolTipText(LOCTEXT("FractureEditor.Visibility_Tooltip", "Toggle showing the bone colours of the geometry collection"))
-					]
-					.OnGetMenuContent(this, &FFractureEditorModeToolkit::GetViewMenuContent)
-				]
-			]
-		]
-
-		+SVerticalBox::Slot()
-		.AutoHeight()
-		[
-			SNew(SExpandableArea)
-			.AreaTitle(FText(LOCTEXT("Cluster", "Cluster")))
-			.HeaderPadding(FMargin(2.0, 2.0))
-			.Padding(FMargin(MorePadding))
-			.BorderImage(FEditorStyle::Get().GetBrush("DetailsView.CategoryTop"))
-			.BorderBackgroundColor( FLinearColor( .6,.6,.6, 1.0f ) )
-			.BodyBorderBackgroundColor (FLinearColor( 1.0, 0.0, 0.0))
-			.AreaTitleFont(FEditorStyle::Get().GetFontStyle("DetailsView.CategoryFontStyle"))
-			.BodyContent()
-			[
-				SNew(SVerticalBox)
-
-				+ SVerticalBox::Slot()
-				.AutoHeight()
-				.HAlign(HAlign_Center)
-				[
-					ClusterToolBuilder.MakeWidget()
-				]
-
-				+SVerticalBox::Slot()
-				.Padding(0.0f, MorePadding, 0.0f, 0.0f)
-				.AutoHeight()
-				.HAlign(HAlign_Center)
-				[
-					SNew(SHorizontalBox)
-
-					+SHorizontalBox::Slot()
-					.HAlign(HAlign_Right)
-					.VAlign(VAlign_Center)
-					.AutoWidth()
-					[
-						SNew(STextBlock)
-						.Text(FText(LOCTEXT("FractureAutoClusterMode", "Auto Cluster Mode")))
-						.Font(FCoreStyle::GetDefaultFontStyle("Regular", 8))
-					]
-
-					+SHorizontalBox::Slot()
-					.Padding(MorePadding, 0.0f)
-					.AutoWidth()
-					[
-						SNew(SComboButton)
-						.OnGetMenuContent(this, &FFractureEditorModeToolkit::GetAutoClusterModesMenu)
-						.IsEnabled(true)
-						.ButtonContent()
-						[
-							SNew(STextBlock)
-							.Text(this, &FFractureEditorModeToolkit::GetAutoClusterModeDisplayName)
-							.Font(FCoreStyle::GetDefaultFontStyle("Regular", 8))
-						]
-					]
-
-					+SHorizontalBox::Slot()
-					.HAlign(HAlign_Right)
-					.VAlign(VAlign_Center)
-					.Padding(MorePadding, 0.0f)
-					.AutoWidth()
-					[
-						SNew(STextBlock)
-						.Text(FText(LOCTEXT("FractureAutoClusterSites", "# Sites")))
-						.Font(FCoreStyle::GetDefaultFontStyle("Regular", 8))
-					]
-
-					+SHorizontalBox::Slot()
-					.AutoWidth()
-					[
-						SNew(SSpinBox<uint32>)
-						.IsEnabled(true)
-						.MinValue(1)
-						.MaxValue(9999)
-						.Value(this, &FFractureEditorModeToolkit::GetAutoClusterSiteCount)
-						.OnValueChanged(this, &FFractureEditorModeToolkit::SetAutoClusterSiteCount)
-						.ToolTipText(LOCTEXT("AutoClusterSitesTooltip", "How many groups you would like to try to autocluster into"))
-
-					]
+					StatisticsExpander
 				]
 			]
-		]
-
-		+SVerticalBox::Slot()
-		[
-			SNew(SSplitter)
-			.Orientation( Orient_Vertical )
 
-			+SSplitter::Slot()
-			.SizeRule(TAttribute<SSplitter::ESizeRule>(this, &FFractureEditorModeToolkit::GetFractureGroupSizeRule ))
+			+ SWidgetSwitcher::Slot()
 			[
 				SNew(SExpandableArea)
 				.AreaTitle(FText(LOCTEXT("Fracture", "Fracture")))
@@ -415,17 +252,10 @@ void FFractureEditorModeToolkit::Init(const TSharedPtr<IToolkitHost>& InitToolki
 				.BorderBackgroundColor( FLinearColor( .6,.6,.6, 1.0f ) )
 				.BodyBorderBackgroundColor (FLinearColor( 1.0, 0.0, 0.0))
 				.AreaTitleFont(FEditorStyle::Get().GetFontStyle("DetailsView.CategoryFontStyle"))
-				.OnAreaExpansionChanged(this, &FFractureEditorModeToolkit::OnFractureGroupExpansionChanged)
 				.BodyContent()
 				[
 					SNew(SScrollBox)
 
-					+SScrollBox::Slot()
-					.HAlign(HAlign_Center)
-					[
-						FractureToolBuilder.MakeWidget()
-					]
-
 					+SScrollBox::Slot()
 					.Padding(0.0f, MorePadding)
 					[
@@ -451,27 +281,21 @@ void FFractureEditorModeToolkit::Init(const TSharedPtr<IToolkitHost>& InitToolki
 
 						+SHorizontalBox::Slot()
 						.FillWidth(1)
-					]
-				]
-			]
+
+						+SHorizontalBox::Slot()
+						.AutoWidth()
+						[
+							SNew(SButton)
+							.HAlign(HAlign_Center)
+							.ContentPadding(FMargin(MorePadding, Padding))
+							.OnClicked_Lambda( [this] () -> FReply { SetActiveTool(0); return FReply::Handled(); } )
+							.Text(FText(LOCTEXT("FractureCancelButton", "Cancel")))
+						]
 
 
-			+SSplitter::Slot()
-			.SizeRule(TAttribute<SSplitter::ESizeRule>(this, &FFractureEditorModeToolkit::GetOutlinerGroupSizeRule ))
-			[
-				SNew(SExpandableArea)
-				.AreaTitle(FText(LOCTEXT("Outliner", "Outliner")))
-				.HeaderPadding(FMargin(2.0, 2.0))
-				.Padding(FMargin(MorePadding))
-				.BorderImage(FEditorStyle::Get().GetBrush("DetailsView.CategoryTop"))
-				.BorderBackgroundColor( FLinearColor( .6,.6,.6, 1.0f ) )
-				.BodyBorderBackgroundColor (FLinearColor( 1.0, 0.0, 0.0))
-				.AreaTitleFont(FEditorStyle::Get().GetFontStyle("DetailsView.CategoryFontStyle"))
-				.OnAreaExpansionChanged(this, &FFractureEditorModeToolkit::OnOutlinerGroupExpansionChanged)
-				.BodyContent()
-				[
-					SAssignNew(OutlinerView, SGeometryCollectionOutliner )
-					.OnBoneSelectionChanged(this, &FFractureEditorModeToolkit::OnOutlinerBoneSelectionChanged)
+						+SHorizontalBox::Slot()
+						.FillWidth(1)
+					]
 				]
 			]
 		]
@@ -481,12 +305,89 @@ void FFractureEditorModeToolkit::Init(const TSharedPtr<IToolkitHost>& InitToolki
 	// Bind Chaos Commands;
 	BindCommands();
 
-	SetActiveTool(GetActiveTool());
-
 	FModeToolkit::Init(InitToolkitHost);
 
 }
 
+const TArray<FName> FFractureEditorModeToolkit::PaletteNames = { FName(TEXT("Fracture")), FName(TEXT("Cluster")) };
+
+FText FFractureEditorModeToolkit::GetToolPaletteDisplayName(FName Palette) 
+{ 
+	return FText::FromName(Palette);
+}
+
+void FFractureEditorModeToolkit::BuildToolPalette(FName PaletteIndex, class FToolBarBuilder& ToolbarBuilder) 
+{
+	const FFractureEditorCommands& Commands = FFractureEditorCommands::Get();
+
+	if (PaletteIndex == PaletteNames[0])
+	{
+		ToolbarBuilder.AddToolBarButton(Commands.GenerateAsset);
+
+		ToolbarBuilder.AddSeparator();
+		ToolbarBuilder.AddToolBarButton(Commands.SelectAll);
+		ToolbarBuilder.AddToolBarButton(Commands.SelectNone);
+		ToolbarBuilder.AddToolBarButton(Commands.SelectNeighbors);
+		ToolbarBuilder.AddToolBarButton(Commands.SelectSiblings);
+		ToolbarBuilder.AddToolBarButton(Commands.SelectAllInCluster);
+		ToolbarBuilder.AddToolBarButton(Commands.SelectInvert);
+
+		ToolbarBuilder.AddSeparator();
+
+		ToolbarBuilder.AddToolBarWidget(ExplodedViewWidget.ToSharedRef(), FText(LOCTEXT("FractureExplodedPercentage", "Explode")));
+		ToolbarBuilder.AddToolBarWidget(LevelViewWidget.ToSharedRef(), LOCTEXT("FractureViewLevel", "Level"));
+		// ToolbarBuilder.AddToolBarWidget(ShowBoneColorsWidget.ToSharedRef(), LOCTEXT("FractureViewOptions", "View"));
+		ToolbarBuilder.AddToolBarButton(Commands.ToggleShowBoneColors);
+
+		ToolbarBuilder.AddSeparator();
+		ToolbarBuilder.AddToolBarButton(Commands.Uniform);
+		ToolbarBuilder.AddToolBarButton(Commands.Clustered);
+		ToolbarBuilder.AddToolBarButton(Commands.Radial);
+		ToolbarBuilder.AddToolBarButton(Commands.Planar);
+		ToolbarBuilder.AddToolBarButton(Commands.Slice);
+		ToolbarBuilder.AddToolBarButton(Commands.Brick);
+	}	
+
+	else if (PaletteIndex == PaletteNames[1])
+	{
+		ToolbarBuilder.AddToolBarButton(Commands.GenerateAsset);
+		ToolbarBuilder.AddSeparator();
+
+		ToolbarBuilder.AddToolBarButton(Commands.SelectAll);
+		ToolbarBuilder.AddToolBarButton(Commands.SelectNone);
+		ToolbarBuilder.AddToolBarButton(Commands.SelectNeighbors);
+		ToolbarBuilder.AddToolBarButton(Commands.SelectSiblings);
+		ToolbarBuilder.AddToolBarButton(Commands.SelectAllInCluster);
+		ToolbarBuilder.AddToolBarButton(Commands.SelectInvert);
+
+		ToolbarBuilder.AddSeparator();
+
+		ToolbarBuilder.AddToolBarWidget(ExplodedViewWidget.ToSharedRef(), FText(LOCTEXT("FractureExplodedPercentage", "Explode")));
+		ToolbarBuilder.AddToolBarWidget(LevelViewWidget.ToSharedRef(), LOCTEXT("FractureViewLevel", "Level"));
+		// ToolbarBuilder.AddToolBarWidget(ShowBoneColorsWidget.ToSharedRef(), LOCTEXT("FractureViewOptions", "View"));
+		ToolbarBuilder.AddToolBarButton(Commands.ToggleShowBoneColors);
+
+		ToolbarBuilder.AddSeparator();
+
+		ToolbarBuilder.AddToolBarButton(Commands.AutoCluster);
+		ToolbarBuilder.AddComboButton(
+			FUIAction(),
+			FOnGetContent::CreateSP(this, &FFractureEditorModeToolkit::GetAutoClusterModesMenu),
+			FText(),
+			LOCTEXT("AutoClusterOptionsToolTip", "Auto Cluster Options")
+		);
+
+		ToolbarBuilder.AddToolBarButton(Commands.Flatten);
+		// ToolbarBuilder.AddToolBarButton(Commands.FlattenToLevel);
+		ToolbarBuilder.AddToolBarButton(Commands.Cluster);
+		ToolbarBuilder.AddToolBarButton(Commands.Uncluster);
+		// ToolbarBuilder.AddToolBarButton(Commands.Merge);
+		ToolbarBuilder.AddToolBarButton(Commands.MoveUp);
+
+	}
+}
+
+
 void FFractureEditorModeToolkit::BindCommands()
 {
 	const FFractureEditorCommands& Commands = FFractureEditorCommands::Get();
@@ -615,10 +516,6 @@ void FFractureEditorModeToolkit::BindCommands()
 				FCanExecuteAction(),//::CreateLambda([]() { return true; }),
 				FIsActionChecked::CreateSP(this, &FFractureEditorModeToolkit::IsActiveTool, FractureTool)
 			);
-
-			// Explicitly set the active tool to the Uniform Fracture
-			if (FractureTool->GetUICommandInfo() == Commands.Uniform)
-				SetActiveTool(FractureTool);
 		}
 	}
 
@@ -653,6 +550,14 @@ void FFractureEditorModeToolkit::BindCommands()
 	);
 }
 
+void FFractureEditorModeToolkit::OnToolPaletteChanged(FName PaletteName) 
+{
+	// make sure if you leave the fracture tab, that the fracture toolkit is also closed
+	if (GetActiveTool() != nullptr && PaletteName != PaletteNames[0])
+	{
+		SetActiveTool(0);
+	}
+}
 
 FName FFractureEditorModeToolkit::GetToolkitFName() const
 {
@@ -838,6 +743,8 @@ TSharedRef<SWidget> FFractureEditorModeToolkit::GetAutoClusterModesMenu()
 {
 	FMenuBuilder MenuBuilder(true, NULL);
 
+	MenuBuilder.BeginSection("AutoClusterMode", LOCTEXT("AutoClusterMode", "Auto Cluster Mode"));
+
 	MenuBuilder.AddMenuEntry(
 		LOCTEXT("AutoClusterModeBoundingBox", "Bounding Box"),
 		LOCTEXT("AutoClusterModeBoundingBoxTooltip", "Use the Bounding Box method when Auto Clustering.  This uses the overlapping of the bones bounding boxes to determine which bones to connect."),
@@ -877,7 +784,38 @@ TSharedRef<SWidget> FFractureEditorModeToolkit::GetAutoClusterModesMenu()
 		EUserInterfaceActionType::ToggleButton
 	);
 
-	//MenuBuilder.EndSection();
+
+	MenuBuilder.EndSection();
+
+	MenuBuilder.BeginSection("AutoClusterOptions", LOCTEXT("AutoClusterOptions", "Auto Cluster Options"));
+	MenuBuilder.AddWidget(
+			SNew(SHorizontalBox)
+
+			+SHorizontalBox::Slot()
+			[
+				SNew(SSpacer)
+			]
+
+			+SHorizontalBox::Slot()
+			.AutoWidth()
+			.HAlign(HAlign_Right)
+			[
+
+				SNew(SSpinBox<uint32>)
+				// .TypeInterface(NumericTypeInterface)
+				.IsEnabled(true)
+				.MinValue(TOptional<uint32>())
+				.MaxValue(TOptional<uint32>())
+				.Value(this, &FFractureEditorModeToolkit::GetAutoClusterSiteCount)
+				.OnValueChanged(this, &FFractureEditorModeToolkit::SetAutoClusterSiteCount)
+				.ToolTipText(LOCTEXT("AutoClusterSitesTooltip", "How many groups you would like to try to autocluster into"))
+				.Style(&FEditorStyle::GetWidgetStyle<FSpinBoxStyle>("Sequencer.HyperlinkSpinBox"))
+			],
+
+		LOCTEXT("FractureAutoClusterSites", "# Sites") 
+	);
+
+	MenuBuilder.EndSection();
 
 	return MenuBuilder.MakeWidget();
 }
@@ -964,15 +902,16 @@ void FFractureEditorModeToolkit::SetActiveTool(UFractureTool* InActiveTool)
 	UFractureCommonSettings* CommonSettings = GetMutableDefault<UFractureCommonSettings>();
 	CommonSettings->OwnerTool = ActiveTool;
 
+	TArray<UObject*> Settings;
 	if (ActiveTool != nullptr)
 	{
-		TArray<UObject*> Settings = ActiveTool->GetSettingsObjects();
+		Settings = ActiveTool->GetSettingsObjects();
 		Settings.Insert(CommonSettings, 0);
 
-		DetailsView->SetObjects(Settings);
-
 		ActiveTool->FractureContextChanged();
 	}
+
+	DetailsView->SetObjects(Settings);
 }
 
 
@@ -1230,6 +1169,8 @@ FReply FFractureEditorModeToolkit::OnFractureClicked()
 
 		SetOutlinerComponents(NewComponents);
 
+		SetActiveTool(0);
+
 		float ProcessingTime = static_cast<float>(FPlatformTime::Seconds() - CacheStartTime);
 
 		GCurrentLevelEditingViewportClient->Invalidate();
@@ -1699,24 +1640,28 @@ class AGeometryCollectionActor* FFractureEditorModeToolkit::CreateNewGeometryAct
 
 void FFractureEditorModeToolkit::ExecuteFracture(FFractureContext& FractureContext)
 {
-	FractureContext.FracturedGeometryCollection->Modify();
-	ActiveTool->ExecuteFracture(FractureContext);
-	TSharedPtr<FGeometryCollection, ESPMode::ThreadSafe> GeometryCollectionPtr = FractureContext.FracturedGeometryCollection->GetGeometryCollection();
-	FGeometryCollection* OutGeometryCollection = GeometryCollectionPtr.Get();
-	FGeometryCollectionClusteringUtility::UpdateHierarchyLevelOfChildren(OutGeometryCollection, -1);
+	if (ActiveTool != nullptr)
+	{
 
-	UGeometryCollectionComponent* GeometryCollectionComponent = Cast<UGeometryCollectionComponent>(FractureContext.OriginalPrimitiveComponent);
+		FractureContext.FracturedGeometryCollection->Modify();
+		ActiveTool->ExecuteFracture(FractureContext);
+		TSharedPtr<FGeometryCollection, ESPMode::ThreadSafe> GeometryCollectionPtr = FractureContext.FracturedGeometryCollection->GetGeometryCollection();
+		FGeometryCollection* OutGeometryCollection = GeometryCollectionPtr.Get();
+		FGeometryCollectionClusteringUtility::UpdateHierarchyLevelOfChildren(OutGeometryCollection, -1);
 
-	if (GeometryCollectionComponent != nullptr) // Create new GC actor from static mesh
-	{
-		FScopedColorEdit EditBoneColor = GeometryCollectionComponent->EditBoneSelection();
+		UGeometryCollectionComponent* GeometryCollectionComponent = Cast<UGeometryCollectionComponent>(FractureContext.OriginalPrimitiveComponent);
 
-		EditBoneColor.SelectBones(GeometryCollection::ESelectionMode::None);
-		SetBoneSelection(GeometryCollectionComponent, EditBoneColor.GetSelectedBones(), true);
-	}
+		if (GeometryCollectionComponent != nullptr) // Create new GC actor from static mesh
+		{
+			FScopedColorEdit EditBoneColor = GeometryCollectionComponent->EditBoneSelection();
 
-	GeometryCollectionComponent->MarkRenderDynamicDataDirty();
-	GeometryCollectionComponent->MarkRenderStateDirty();
+			EditBoneColor.SelectBones(GeometryCollection::ESelectionMode::None);
+			SetBoneSelection(GeometryCollectionComponent, EditBoneColor.GetSelectedBones(), true);
+		}
+
+		GeometryCollectionComponent->MarkRenderDynamicDataDirty();
+		GeometryCollectionComponent->MarkRenderStateDirty();
+	}
 }
 
 
@@ -1974,4 +1919,77 @@ void FFractureEditorModeToolkit::OnOutlinerBoneSelectionChanged(UGeometryCollect
 }
 
 
+FText FFractureEditorModeToolkit::GetStatisticsSummary() const
+{
+
+	TArray<const FGeometryCollection*> GeometryCollectionArray;
+	if (USelection* SelectedActors = GEditor->GetSelectedActors())
+	{
+		for (FSelectionIterator Iter(*SelectedActors); Iter; ++Iter)
+		{
+			if (AGeometryCollectionActor* Actor = Cast<AGeometryCollectionActor>(*Iter))
+			{
+				const UGeometryCollection* RestCollection = Actor->GetGeometryCollectionComponent()->GetRestCollection();
+				const FGeometryCollection* GeometryCollection = RestCollection->GetGeometryCollection().Get();
+
+				if (GeometryCollection != nullptr)
+				{
+					GeometryCollectionArray.Add(GeometryCollection);
+				}
+			}
+		}
+	}
+
+
+	FString Buffer;
+
+	if (GeometryCollectionArray.Num() > 0)
+	{
+		TArray<int32> LevelTransformsAll;
+		LevelTransformsAll.SetNumZeroed(10);
+		int32 LevelMax = INT_MIN;
+
+		for (int32 Idx = 0; Idx < GeometryCollectionArray.Num(); ++Idx)
+		{
+			const FGeometryCollection* GeometryCollection = GeometryCollectionArray[Idx];
+
+			check(GeometryCollection);
+
+
+			Buffer += FString::Printf(TEXT("Sum of the selected Geometry Collections\n\n"));
+
+			const TManagedArray<int32>& Levels = GeometryCollection->GetAttribute<int32>("Level", FGeometryCollection::TransformGroup);
+
+			TArray<int32> LevelTransforms;
+			for (int32 Element = 0, NumElement = Levels.Num(); Element < NumElement; ++Element)
+			{
+				const int32 NodeLevel = Levels[Element];
+				while (LevelTransforms.Num() <= NodeLevel)
+				{
+					LevelTransforms.SetNum(NodeLevel + 1);
+					LevelTransforms[NodeLevel] = 0;
+				}
+				++LevelTransforms[NodeLevel];
+			}
+
+			for (int32 Level = 0; Level < LevelTransforms.Num(); ++Level)
+			{
+				LevelTransformsAll[Level] += LevelTransforms[Level];
+			}
+
+			if (LevelTransforms.Num() > LevelMax)
+			{
+				LevelMax = LevelTransforms.Num();
+			}
+		}
+
+		for (int32 Level = 0; Level < LevelMax; ++Level)
+		{
+			Buffer += FString::Printf(TEXT("Level: %d \t - \t %d\n"), Level, LevelTransformsAll[Level]);
+		}
+	}
+
+	return FText::FromString(Buffer);
+}
+
 #undef LOCTEXT_NAMESPACE
diff --git a/Engine/Plugins/Experimental/ChaosEditor/Source/FractureEditor/Private/FractureEditorStyle.cpp b/Engine/Plugins/Experimental/ChaosEditor/Source/FractureEditor/Private/FractureEditorStyle.cpp
--- a/Engine/Plugins/Experimental/ChaosEditor/Source/FractureEditor/Private/FractureEditorStyle.cpp
+++ b/Engine/Plugins/Experimental/ChaosEditor/Source/FractureEditor/Private/FractureEditorStyle.cpp
@@ -11,8 +11,8 @@ FName FFractureEditorStyle::StyleName("FractureEditorStyle");
 FFractureEditorStyle::FFractureEditorStyle()
 	: FSlateStyleSet(StyleName)
 {
-	const FVector2D IconSize(48.0f, 48.0f);
-	const FVector2D SmallIconSize(32.0f, 32.0f);
+	const FVector2D IconSize(20.0f, 20.0f);
+	const FVector2D SmallIconSize(20.0f, 20.0f);
 	const FVector2D LabelIconSize(16.0f, 16.0f);
 
 
@@ -71,6 +71,8 @@ FFractureEditorStyle::FFractureEditorStyle()
 	Set("FractureEditor.Exploded", new FSlateImageBrush(RootToContentDir(TEXT("MiniExploded_32x.png")), LabelIconSize));
 	Set("FractureEditor.Levels", new FSlateImageBrush(RootToContentDir(TEXT("MiniLevel_32x.png")), LabelIconSize));
 	Set("FractureEditor.Visibility", new FSlateImageBrush(RootToContentDir(TEXT("GeneralVisibility_48x.png")), SmallIconSize));
+	Set("FractureEditor.ToggleShowBoneColors", new FSlateImageBrush(RootToContentDir(TEXT("GeneralVisibility_48x.png")), IconSize));
+	Set("FractureEditor.ToggleShowBoneColors.Small", new FSlateImageBrush(RootToContentDir(TEXT("GeneralVisibility_48x.png")), SmallIconSize));
 	Set("FractureEditor.ViewUpOneLevel", new FSlateImageBrush(RootToContentDir(TEXT("LevelViewUp_48x.png")), IconSize));
 	Set("FractureEditor.ViewUpOneLevel.Small", new FSlateImageBrush(RootToContentDir(TEXT("LevelViewUp_48x.png")), SmallIconSize));
 	Set("FractureEditor.ViewDownOneLevel", new FSlateImageBrush(RootToContentDir(TEXT("LevelViewDown_48x.png")), IconSize));
diff --git a/Engine/Plugins/Experimental/ChaosEditor/Source/FractureEditor/Public/FractureEditorModeToolkit.h b/Engine/Plugins/Experimental/ChaosEditor/Source/FractureEditor/Public/FractureEditorModeToolkit.h
--- a/Engine/Plugins/Experimental/ChaosEditor/Source/FractureEditor/Public/FractureEditorModeToolkit.h
+++ b/Engine/Plugins/Experimental/ChaosEditor/Source/FractureEditor/Public/FractureEditorModeToolkit.h
@@ -95,6 +95,7 @@ public:
 	// static int32 GetLevelCount();
 	int32 GetLevelCount();
 
+	FText GetStatisticsSummary() const;
 
 	TSharedRef<SWidget> GetAutoClusterModesMenu();	
 	void SetAutoClusterMode(EFractureAutoClusterMode InAutoClusterMode);
@@ -103,14 +104,16 @@ public:
 	void SetAutoClusterSiteCount(uint32 InSiteCount);
 	uint32 GetAutoClusterSiteCount() const;
 
-	/** Fracture Group Layout Management */
-	SSplitter::ESizeRule GetFractureGroupSizeRule() const { return bFractureGroupExpanded ? SSplitter::ESizeRule::FractionOfParent : SSplitter::ESizeRule::SizeToContent; }
-	void OnFractureGroupExpansionChanged(bool bExpanded) { bFractureGroupExpanded = bExpanded; }
-
-	/** Outliner Group Layout Management */
-	SSplitter::ESizeRule GetOutlinerGroupSizeRule() const { return bOutlinerGroupExpanded ? SSplitter::ESizeRule::FractionOfParent : SSplitter::ESizeRule::SizeToContent; }
-	void OnOutlinerGroupExpansionChanged(bool bExpanded) { bOutlinerGroupExpanded = bExpanded; }
+	/** Returns the number of Mode specific tabs in the mode toolbar **/ 
+	const static TArray<FName> PaletteNames;
+	virtual void GetToolPaletteNames( TArray<FName>& InPaletteName ) const { InPaletteName = PaletteNames; }
+	virtual FText GetToolPaletteDisplayName(FName PaletteName); 
+	virtual void BuildToolPalette(FName PaletteName, class FToolBarBuilder& ToolbarBuilder);
+	virtual void OnToolPaletteChanged(FName PaletteName) override;
 
+	TSharedPtr<SWidget> ExplodedViewWidget;
+	TSharedPtr<SWidget> LevelViewWidget;
+	TSharedPtr<SWidget> ShowBoneColorsWidget;
 
 protected:
 	static bool IsGeometryCollectionSelected();
@@ -129,9 +132,6 @@ private:
 
 	UFractureTool* ActiveTool;
 
-	bool bFractureGroupExpanded = true;
-	bool bOutlinerGroupExpanded = true;
-	
 	TSharedPtr<IDetailsView> DetailsView;
 	TSharedPtr<SWidget> ToolkitWidget;
 	TSharedPtr<SGeometryCollectionOutliner> OutlinerView;
