Fracture: Promote AutoCluster to a Fracture Tool

Also fixed LevelView Menu

diff --git a/Engine/Plugins/Experimental/ChaosEditor/Source/FractureEditor/Private/FractureEditorCommands.cpp b/Engine/Plugins/Experimental/ChaosEditor/Source/FractureEditor/Private/FractureEditorCommands.cpp
--- a/Engine/Plugins/Experimental/ChaosEditor/Source/FractureEditor/Private/FractureEditorCommands.cpp
+++ b/Engine/Plugins/Experimental/ChaosEditor/Source/FractureEditor/Private/FractureEditorCommands.cpp
@@ -34,7 +34,6 @@ void FFractureEditorCommands::RegisterCommands()
 	// Cluster Operations
 	UI_COMMAND(Flatten, "Flatten", "Flattens all bones to level 1", EUserInterfaceActionType::Button, FInputChord() );
 	UI_COMMAND(FlattenToLevel, "FlattenToLevel", "Flatten Only up to the Current View Level", EUserInterfaceActionType::Button, FInputChord() );
-	UI_COMMAND(AutoCluster, "Auto", "Auto Cluster", EUserInterfaceActionType::Button, FInputChord() );
 	UI_COMMAND(Cluster, "Cluster", "Clusters selected bones under a new parent", EUserInterfaceActionType::Button, FInputChord() );
 	UI_COMMAND(Uncluster, "Uncluster", "Uncluster", EUserInterfaceActionType::Button, FInputChord() );
 	UI_COMMAND(Merge, "Merge", "Merge", EUserInterfaceActionType::Button, FInputChord() );
diff --git a/Engine/Plugins/Experimental/ChaosEditor/Source/FractureEditor/Private/FractureEditorModeToolkit.cpp b/Engine/Plugins/Experimental/ChaosEditor/Source/FractureEditor/Private/FractureEditorModeToolkit.cpp
--- a/Engine/Plugins/Experimental/ChaosEditor/Source/FractureEditor/Private/FractureEditorModeToolkit.cpp
+++ b/Engine/Plugins/Experimental/ChaosEditor/Source/FractureEditor/Private/FractureEditorModeToolkit.cpp
@@ -82,8 +82,6 @@ TArray<UClass*> FindFractureToolClasses()
 FFractureEditorModeToolkit::FFractureEditorModeToolkit()
 	: ExplodeAmount(0.0f)
 	, FractureLevel(-1)
-	, AutoClusterMode(EFractureAutoClusterMode::BoundingBox)
-	, AutoClusterSiteCount(10)
 	, ActiveTool(nullptr)
 {
 }
@@ -137,24 +135,30 @@ void FFractureEditorModeToolkit::Init(const TSharedPtr<IToolkitHost>& InitToolki
 	.ContentPadding(0)
 	.ButtonStyle(FEditorStyle::Get(), "Toolbar.Button")
 	.ForegroundColor(FEditorStyle::Get().GetSlateColor("ToolBar.SToolBarComboButtonBlock.ComboButton.Color"))
+	.OnGetMenuContent(this, &FFractureEditorModeToolkit::GetLevelViewMenuContent) 
 	.ButtonContent()
 	[
-		SNew(STextBlock)
-		.Text_Lambda( [=]() -> FText {
-			if (FractureLevel < 0)
-			{
-				return LOCTEXT("FractureViewAllLevels", "All");
-			}
-			else if (FractureLevel == 0)
-			{
-				return LOCTEXT("FractureViewRootLevel", "Root");
-			}
+		SNew(SBox)
+		.WidthOverride(36)
+		[
+			SNew(STextBlock)
+			.Justification(ETextJustify::Center)
+			.Text_Lambda( [=]() -> FText {
+				if (FractureLevel < 0)
+				{
+					return LOCTEXT("FractureViewAllLevels", "All");
+				}
+				else if (FractureLevel == 0)
+				{
+					return LOCTEXT("FractureViewRootLevel", "Root");
+				}
 
-			return FText::Format(NSLOCTEXT("FractureEditor", "CurrentLevel", "{0}"), FText::AsNumber(FractureLevel));
+				return FText::Format(NSLOCTEXT("FractureEditor", "CurrentLevel", "{0}"), FText::AsNumber(FractureLevel));
 
-		})
-		.Font(FCoreStyle::GetDefaultFontStyle("Regular", 12))
-		.ToolTipText(LOCTEXT("FractureEditor.Level_Tooltip", "Set the currently view level of the geometry collection"))
+			})
+			.Font(FCoreStyle::GetDefaultFontStyle("Regular", 12))
+			.ToolTipText(LOCTEXT("FractureEditor.Level_Tooltip", "Set the currently view level of the geometry collection"))
+		]
 
 	];
 
@@ -244,64 +248,55 @@ void FFractureEditorModeToolkit::Init(const TSharedPtr<IToolkitHost>& InitToolki
 
 			+ SWidgetSwitcher::Slot()
 			[
-				SNew(SExpandableArea)
-				.AreaTitle(FText(LOCTEXT("Fracture", "Fracture")))
-				.HeaderPadding(FMargin(2.0, 2.0))
-				.Padding(FMargin(MorePadding))
-				.BorderImage(FEditorStyle::Get().GetBrush("DetailsView.CategoryTop"))
-				.BorderBackgroundColor( FLinearColor( .6,.6,.6, 1.0f ) )
-				.BodyBorderBackgroundColor (FLinearColor( 1.0, 0.0, 0.0))
-				.AreaTitleFont(FEditorStyle::Get().GetFontStyle("DetailsView.CategoryFontStyle"))
-				.BodyContent()
+
+				SNew(SScrollBox)
+
+				+SScrollBox::Slot()
+				.Padding(0.0f, MorePadding)
+				[
+					DetailsView.ToSharedRef()
+				]
+
+				+SScrollBox::Slot()
+				.Padding(16.f)
 				[
-					SNew(SScrollBox)
+					SNew(SHorizontalBox)
 
-					+SScrollBox::Slot()
-					.Padding(0.0f, MorePadding)
+					+SHorizontalBox::Slot()
+					.FillWidth(1)
+
+					+SHorizontalBox::Slot()
+					.AutoWidth()
 					[
-						DetailsView.ToSharedRef()
+						SNew( SButton )
+						.HAlign(HAlign_Center)
+						.ContentPadding(FMargin(MorePadding, Padding))
+						.OnClicked(this, &FFractureEditorModeToolkit::OnFractureClicked)
+						.IsEnabled( this, &FFractureEditorModeToolkit::CanExecuteFracture)
+						.Text_Lambda( [this] () -> FText { return ActiveTool ? ActiveTool->GetApplyText() :  LOCTEXT("FractureApplyButton", "Apply"); })
 					]
 
-					+SScrollBox::Slot()
+					+SHorizontalBox::Slot()
+					.FillWidth(1)
+
+					+SHorizontalBox::Slot()
+					.AutoWidth()
 					[
-						SNew(SHorizontalBox)
-						+SHorizontalBox::Slot()
-						.FillWidth(1)
-
-						+SHorizontalBox::Slot()
-						.AutoWidth()
-						[
-							SNew( SButton )
-							.HAlign(HAlign_Center)
-							.ContentPadding(FMargin(MorePadding, Padding))
-							.OnClicked(this, &FFractureEditorModeToolkit::OnFractureClicked)
-							.IsEnabled( this, &FFractureEditorModeToolkit::CanExecuteFracture)
-							.Text(FText(LOCTEXT("FractureFractureButton", "Fracture")))
-						]
-
-						+SHorizontalBox::Slot()
-						.FillWidth(1)
-
-						+SHorizontalBox::Slot()
-						.AutoWidth()
-						[
-							SNew(SButton)
-							.HAlign(HAlign_Center)
-							.ContentPadding(FMargin(MorePadding, Padding))
-							.OnClicked_Lambda( [this] () -> FReply { SetActiveTool(0); return FReply::Handled(); } )
-							.Text(FText(LOCTEXT("FractureCancelButton", "Cancel")))
-						]
-
-
-						+SHorizontalBox::Slot()
-						.FillWidth(1)
+						SNew(SButton)
+						.HAlign(HAlign_Center)
+						.ContentPadding(FMargin(MorePadding, Padding))
+						.OnClicked_Lambda( [this] () -> FReply { SetActiveTool(0); return FReply::Handled(); } )
+						.Text(FText(LOCTEXT("FractureCancelButton", "Cancel")))
 					]
+
+
+					+SHorizontalBox::Slot()
+					.FillWidth(1)
 				]
 			]
 		]
 	];
 
-
 	// Bind Chaos Commands;
 	BindCommands();
 
@@ -322,6 +317,9 @@ void FFractureEditorModeToolkit::BuildToolPalette(FName PaletteIndex, class FToo
 
 	if (PaletteIndex == PaletteNames[0])
 	{
+
+		ToolbarBuilder.AddWidget(SNew(SBox).WidthOverride(4));
+
 		ToolbarBuilder.AddToolBarButton(Commands.GenerateAsset);
 
 		ToolbarBuilder.AddSeparator();
@@ -350,6 +348,7 @@ void FFractureEditorModeToolkit::BuildToolPalette(FName PaletteIndex, class FToo
 
 	else if (PaletteIndex == PaletteNames[1])
 	{
+		ToolbarBuilder.AddWidget(SNew(SBox).WidthOverride(4));
 		ToolbarBuilder.AddToolBarButton(Commands.GenerateAsset);
 		ToolbarBuilder.AddSeparator();
 
@@ -370,12 +369,7 @@ void FFractureEditorModeToolkit::BuildToolPalette(FName PaletteIndex, class FToo
 		ToolbarBuilder.AddSeparator();
 
 		ToolbarBuilder.AddToolBarButton(Commands.AutoCluster);
-		ToolbarBuilder.AddComboButton(
-			FUIAction(),
-			FOnGetContent::CreateSP(this, &FFractureEditorModeToolkit::GetAutoClusterModesMenu),
-			FText(),
-			LOCTEXT("AutoClusterOptionsToolTip", "Auto Cluster Options")
-		);
+		ToolbarBuilder.AddSeparator();
 
 		ToolbarBuilder.AddToolBarButton(Commands.Flatten);
 		// ToolbarBuilder.AddToolBarButton(Commands.FlattenToLevel);
@@ -452,12 +446,6 @@ void FFractureEditorModeToolkit::BindCommands()
 		FCanExecuteAction::CreateStatic(&FFractureEditorModeToolkit::IsGeometryCollectionSelected)
 		);
 
-	ToolkitCommands->MapAction(
-		Commands.AutoCluster,
-		FExecuteAction::CreateSP(this, &FFractureEditorModeToolkit::OnAutoCluster),
-		FCanExecuteAction::CreateStatic(&FFractureEditorModeToolkit::IsGeometryCollectionSelected)
-	);
-
 	ToolkitCommands->MapAction(
 		Commands.Cluster,
 		FExecuteAction::CreateSP(this, &FFractureEditorModeToolkit::OnCluster),
@@ -552,8 +540,7 @@ void FFractureEditorModeToolkit::BindCommands()
 
 void FFractureEditorModeToolkit::OnToolPaletteChanged(FName PaletteName) 
 {
-	// make sure if you leave the fracture tab, that the fracture toolkit is also closed
-	if (GetActiveTool() != nullptr && PaletteName != PaletteNames[0])
+	if (GetActiveTool() != nullptr)
 	{
 		SetActiveTool(0);
 	}
@@ -739,118 +726,6 @@ void FFractureEditorModeToolkit::ViewDownOneLevel()
 	OnSetLevelViewValue(NewLevel);
 }
 
-TSharedRef<SWidget> FFractureEditorModeToolkit::GetAutoClusterModesMenu()
-{
-	FMenuBuilder MenuBuilder(true, NULL);
-
-	MenuBuilder.BeginSection("AutoClusterMode", LOCTEXT("AutoClusterMode", "Auto Cluster Mode"));
-
-	MenuBuilder.AddMenuEntry(
-		LOCTEXT("AutoClusterModeBoundingBox", "Bounding Box"),
-		LOCTEXT("AutoClusterModeBoundingBoxTooltip", "Use the Bounding Box method when Auto Clustering.  This uses the overlapping of the bones bounding boxes to determine which bones to connect."),
-		FSlateIcon(),
-		FUIAction(
-			FExecuteAction::CreateSP( this, &FFractureEditorModeToolkit::SetAutoClusterMode, EFractureAutoClusterMode::BoundingBox),
-			FCanExecuteAction(),
-			FIsActionChecked::CreateLambda( [=]() -> bool { return AutoClusterMode == EFractureAutoClusterMode::BoundingBox; } )
-		),
-		NAME_None,
-		EUserInterfaceActionType::ToggleButton
-	);
-
-	MenuBuilder.AddMenuEntry(
-		LOCTEXT("AutoClusterModeProximity", "Proximity"),
-		LOCTEXT("AutoClusterModeProximityTooltip", "Use the Proximity method when Auto Clustering.  This uses actual bone edge connectivity to determine which bones to connect."),
-		FSlateIcon(),
-		FUIAction(
-			FExecuteAction::CreateSP( this, &FFractureEditorModeToolkit::SetAutoClusterMode, EFractureAutoClusterMode::Proximity),
-			FCanExecuteAction(),
-			FIsActionChecked::CreateLambda( [=]() -> bool { return AutoClusterMode == EFractureAutoClusterMode::Proximity; } )
-		),
-		NAME_None,
-		EUserInterfaceActionType::ToggleButton
-	);
-
-	MenuBuilder.AddMenuEntry(
-		LOCTEXT("AutoClusterModeDistance", "Distance"),
-		LOCTEXT("AutoClusterModeDistanceTooltip", "Use the Distance method when Auto Clustering.  This uses a simple distance calculation to determine which bones to connect"),
-		FSlateIcon(),
-		FUIAction(
-			FExecuteAction::CreateSP( this, &FFractureEditorModeToolkit::SetAutoClusterMode, EFractureAutoClusterMode::Distance),
-			FCanExecuteAction(),
-			FIsActionChecked::CreateLambda( [=]() -> bool { return AutoClusterMode == EFractureAutoClusterMode::Distance; } )
-		),
-		NAME_None,
-		EUserInterfaceActionType::ToggleButton
-	);
-
-
-	MenuBuilder.EndSection();
-
-	MenuBuilder.BeginSection("AutoClusterOptions", LOCTEXT("AutoClusterOptions", "Auto Cluster Options"));
-	MenuBuilder.AddWidget(
-			SNew(SHorizontalBox)
-
-			+SHorizontalBox::Slot()
-			[
-				SNew(SSpacer)
-			]
-
-			+SHorizontalBox::Slot()
-			.AutoWidth()
-			.HAlign(HAlign_Right)
-			[
-
-				SNew(SSpinBox<uint32>)
-				// .TypeInterface(NumericTypeInterface)
-				.IsEnabled(true)
-				.MinValue(TOptional<uint32>())
-				.MaxValue(TOptional<uint32>())
-				.Value(this, &FFractureEditorModeToolkit::GetAutoClusterSiteCount)
-				.OnValueChanged(this, &FFractureEditorModeToolkit::SetAutoClusterSiteCount)
-				.ToolTipText(LOCTEXT("AutoClusterSitesTooltip", "How many groups you would like to try to autocluster into"))
-				.Style(&FEditorStyle::GetWidgetStyle<FSpinBoxStyle>("Sequencer.HyperlinkSpinBox"))
-			],
-
-		LOCTEXT("FractureAutoClusterSites", "# Sites") 
-	);
-
-	MenuBuilder.EndSection();
-
-	return MenuBuilder.MakeWidget();
-}
-
-void FFractureEditorModeToolkit::SetAutoClusterMode(EFractureAutoClusterMode InAutoClusterMode)
-{
-	AutoClusterMode = InAutoClusterMode;
-}
-
-
-FText FFractureEditorModeToolkit::GetAutoClusterModeDisplayName() const
-{
-	switch (AutoClusterMode)
-	{
-		case EFractureAutoClusterMode::Proximity: return LOCTEXT("AutoClusterProximity", "Proximity");
-		case EFractureAutoClusterMode::Distance: return LOCTEXT("AutoClusterDistance", "Distance");
-		case EFractureAutoClusterMode::BoundingBox:
-		default:
-			return LOCTEXT("AutoClusterBoundingBox", "Bounding Box");
-	}
-	return LOCTEXT("AutoClusterBoundingBox", "Bounding Box");
-}
-
-
-void FFractureEditorModeToolkit::SetAutoClusterSiteCount(uint32 InSiteCount)
-{
-	if (InSiteCount >= 1)
-		AutoClusterSiteCount = InSiteCount;
-}
-
-uint32 FFractureEditorModeToolkit::GetAutoClusterSiteCount() const
-{
-	return AutoClusterSiteCount;
-}
-
 TSharedRef<SWidget> FFractureEditorModeToolkit::GetLevelViewMenuContent()
 {
 	FMenuBuilder MenuBuilder(true, GetToolkitCommands());
@@ -906,7 +781,6 @@ void FFractureEditorModeToolkit::SetActiveTool(UFractureTool* InActiveTool)
 	if (ActiveTool != nullptr)
 	{
 		Settings = ActiveTool->GetSettingsObjects();
-		Settings.Insert(CommonSettings, 0);
 
 		ActiveTool->FractureContextChanged();
 	}
@@ -1181,31 +1055,32 @@ FReply FFractureEditorModeToolkit::OnFractureClicked()
 
 bool FFractureEditorModeToolkit::CanExecuteFracture() const
 {
-	if((ActiveTool == nullptr) || !ActiveTool->CanExecuteFracture())
+
+	if (!IsSelectedActorsInEditorWorld())
 	{
 		return false;
 	}
 
-	if (!IsSeletedActorsInEditorWorld())
+	if (!IsGeometryCollectionSelected())
 	{
 		return false;
 	}
 
-	if (IsGeometryCollectionSelected())
+	if (IsStaticMeshSelected())
 	{
-		return IsLeafBoneSelected();
+		return false;
 	}
 
-	if (IsStaticMeshSelected())
+	if (ActiveTool != nullptr) 
 	{
-		return false;
+		return ActiveTool->CanExecuteFracture();
 	}
 	
 	return false;
 }
 
 
-bool FFractureEditorModeToolkit::IsLeafBoneSelected() const
+bool FFractureEditorModeToolkit::IsLeafBoneSelected() 
 {
 	TSet<UGeometryCollectionComponent*> GeomCompSelection;
 	GetSelectedGeometryCollectionComponents(GeomCompSelection);
@@ -1270,37 +1145,6 @@ void FFractureEditorModeToolkit::GetSelectedGeometryCollectionComponents(TSet<UG
 	}
 }
 
-
-void FFractureEditorModeToolkit::OnAutoCluster()
-{
-	SetOutlinerComponents({});
-
-	TSet<UGeometryCollectionComponent*> GeomCompSelection;
-	GetSelectedGeometryCollectionComponents(GeomCompSelection);
-	for (UGeometryCollectionComponent* GeometryCollectionComponent : GeomCompSelection)
-	{
-		FGeometryCollectionEdit GCEdit = GeometryCollectionComponent->EditRestCollection();
-		if (UGeometryCollection* GCObject = GCEdit.GetRestCollection())
-		{
-			TSharedPtr<FGeometryCollection, ESPMode::ThreadSafe> GeometryCollectionPtr = GCObject->GetGeometryCollection();
-			if (FGeometryCollection* GeometryCollection = GeometryCollectionPtr.Get())
-			{
-				FScopedColorEdit EditBoneColor = GeometryCollectionComponent->EditBoneSelection();
-				int32 CurrentLevelView = EditBoneColor.GetViewLevel();
-				UAutoClusterFractureCommand::ClusterChildBonesOfASingleMesh(GeometryCollectionComponent, AutoClusterMode, GetAutoClusterSiteCount());
-				FGeometryCollectionClusteringUtility::UpdateHierarchyLevelOfChildren(GeometryCollection, -1);
-
-				EditBoneColor.ResetBoneSelection();
-				EditBoneColor.SetLevelViewMode(0);
-				EditBoneColor.SetLevelViewMode(CurrentLevelView);
-			}
-		}
-		GeometryCollectionComponent->MarkRenderDynamicDataDirty();
-		GeometryCollectionComponent->MarkRenderStateDirty();
-	}
-	SetOutlinerComponents(GeomCompSelection.Array());
-}
-
 void FFractureEditorModeToolkit::OnCluster()
 {
 	TSet<UGeometryCollectionComponent*> GeomCompSelection;
@@ -1703,7 +1547,7 @@ bool FFractureEditorModeToolkit::IsStaticMeshSelected()
 	return false;
 }
 
-bool FFractureEditorModeToolkit::IsSeletedActorsInEditorWorld()
+bool FFractureEditorModeToolkit::IsSelectedActorsInEditorWorld()
 {
 	USelection* SelectedActors = GEditor->GetSelectedActors();
 	for (FSelectionIterator Iter(*SelectedActors); Iter; ++Iter)
diff --git a/Engine/Plugins/Experimental/ChaosEditor/Source/FractureEditor/Private/FractureToolAutoCluster.cpp b/Engine/Plugins/Experimental/ChaosEditor/Source/FractureEditor/Private/FractureToolAutoCluster.cpp
new file mode 100644
--- /dev/null
+++ b/Engine/Plugins/Experimental/ChaosEditor/Source/FractureEditor/Private/FractureToolAutoCluster.cpp
@@ -0,0 +1,74 @@
+// Copyright 1998-2019 Epic Games, Inc. All Rights Reserved.
+
+#include "FractureToolAutoCluster.h"
+
+#include "FractureEditorStyle.h"
+#include "FractureEditorCommands.h"
+
+#include "AutoClusterFracture.h"
+
+#include "GeometryCollection/GeometryCollectionComponent.h"
+
+#define LOCTEXT_NAMESPACE "FractureAutoCluster"
+
+
+UFractureToolAutoCluster::UFractureToolAutoCluster(const FObjectInitializer& ObjInit)
+  : Super(ObjInit) 
+{
+	Settings = NewObject<UFractureAutoClusterSettings>(GetTransientPackage(), UFractureAutoClusterSettings::StaticClass());
+}
+
+
+// UFractureTool Interface
+FText UFractureToolAutoCluster::GetDisplayText() const
+{
+	return FText(NSLOCTEXT("Fracture", "FractureToolAutoCluster", "Auto")); 
+}
+
+
+FText UFractureToolAutoCluster::GetTooltipText() const
+{
+	return FText(NSLOCTEXT("Fracture", "FractureToolAutoClusterToolTip", "Auto Cluster")); 
+}
+
+FText UFractureToolAutoCluster::GetApplyText() const 
+{ 
+	return FText(NSLOCTEXT("Fracture", "ExecuteAutoCluster", "Auto Cluster")); 
+}
+
+FSlateIcon UFractureToolAutoCluster::GetToolIcon() const
+{
+	return FSlateIcon("FractureEditorStyle", "FractureEditor.AutoCluster");
+}
+
+
+TArray<UObject*> UFractureToolAutoCluster::GetSettingsObjects() const 
+{
+	TArray<UObject*> AllSettings; 
+	AllSettings.Add(Settings);
+	return AllSettings;
+}
+
+void UFractureToolAutoCluster::RegisterUICommand( FFractureEditorCommands* BindingContext )
+{
+	UI_COMMAND_EXT( BindingContext, UICommandInfo, "AutoCluster", "Auto", "Auto Cluster", EUserInterfaceActionType::ToggleButton, FInputChord() );
+	BindingContext->AutoCluster = UICommandInfo;
+}
+
+void UFractureToolAutoCluster::ExecuteFracture(const FFractureContext& FractureContext)
+{
+	if (UGeometryCollectionComponent* GeometryCollectionComponent = Cast<UGeometryCollectionComponent>(FractureContext.OriginalPrimitiveComponent))
+	{
+		FScopedColorEdit EditBoneColor = GeometryCollectionComponent->EditBoneSelection();
+		int32 CurrentLevelView = EditBoneColor.GetViewLevel();
+
+		UAutoClusterFractureCommand::ClusterChildBonesOfASingleMesh(GeometryCollectionComponent, Settings->AutoClusterMode, Settings->SiteCount);
+
+		EditBoneColor.ResetBoneSelection();
+		EditBoneColor.SetLevelViewMode(0);
+		EditBoneColor.SetLevelViewMode(CurrentLevelView);
+	}
+}
+
+
+
diff --git a/Engine/Plugins/Experimental/ChaosEditor/Source/FractureEditor/Private/FractureToolAutoCluster.h b/Engine/Plugins/Experimental/ChaosEditor/Source/FractureEditor/Private/FractureToolAutoCluster.h
new file mode 100644
--- /dev/null
+++ b/Engine/Plugins/Experimental/ChaosEditor/Source/FractureEditor/Private/FractureToolAutoCluster.h
@@ -0,0 +1,55 @@
+// Copyright 1998-2019 Epic Games, Inc. All Rights Reserved.
+
+#pragma once
+
+#include "FractureTool.h"
+#include "AutoClusterFracture.h"
+
+#include "FractureToolAutoCluster.generated.h"
+
+
+UCLASS()
+class UFractureAutoClusterSettings
+	: public UObject
+{
+	GENERATED_BODY()
+
+public:
+
+	UFractureAutoClusterSettings()
+		: AutoClusterMode(EFractureAutoClusterMode::BoundingBox)
+		, SiteCount(10)
+	{}
+
+	UPROPERTY(EditAnywhere, Category = AutoCluster, meta = (DisplayName = "Mode"))
+	EFractureAutoClusterMode AutoClusterMode;
+
+	UPROPERTY(EditAnywhere, Category = AutoCluster, meta = (DisplayName = "Cluster Sites"))// , UIMin = "1", UIMax = "5000", ClampMin = "1"))
+	uint32 SiteCount;
+
+};
+
+
+UCLASS(DisplayName="AutoCluster", Category="FractureTools")
+class UFractureToolAutoCluster: public UFractureTool
+{
+public:
+	GENERATED_BODY()
+
+	UFractureToolAutoCluster(const FObjectInitializer& ObjInit);
+
+	// UFractureTool Interface
+	virtual FText GetDisplayText() const override;
+	virtual FText GetTooltipText() const override;
+	virtual FText GetApplyText() const override; 
+	virtual FSlateIcon GetToolIcon() const override;
+	virtual TArray<UObject*> GetSettingsObjects() const override;
+
+	virtual void RegisterUICommand( FFractureEditorCommands* BindingContext );
+
+	virtual void ExecuteFracture(const FFractureContext& Context) override;
+
+	UPROPERTY(EditAnywhere, Category = AutoCluster)
+	UFractureAutoClusterSettings* Settings;
+
+};
\ No newline at end of file
diff --git a/Engine/Plugins/Experimental/ChaosEditor/Source/FractureEditor/Private/FractureToolBitmap.cpp b/Engine/Plugins/Experimental/ChaosEditor/Source/FractureEditor/Private/FractureToolBitmap.cpp
--- a/Engine/Plugins/Experimental/ChaosEditor/Source/FractureEditor/Private/FractureToolBitmap.cpp
+++ b/Engine/Plugins/Experimental/ChaosEditor/Source/FractureEditor/Private/FractureToolBitmap.cpp
@@ -3,6 +3,7 @@
 
 #include "FractureToolBitmap.h"
 
+#include "FractureEditorModeToolkit.h"
 #include "FractureEditorStyle.h"
 
 #define LOCTEXT_NAMESPACE "FractureBitmap"
@@ -38,8 +39,15 @@ void UFractureToolBitmap::RegisterUICommand( FFractureEditorCommands* BindingCon
 TArray<UObject*> UFractureToolBitmap::GetSettingsObjects() const 
 { 
 	TArray<UObject*> Settings; 
+	Settings.Add(GetMutableDefault<UFractureCommonSettings>());
 	Settings.Add(GetMutableDefault<UFractureBitmapSettings>());
 	return Settings;
 }
 
+
+bool UFractureToolBitmap::CanExecuteFracture() const
+{
+	return FFractureEditorModeToolkit::IsLeafBoneSelected();
+}
+
 #undef LOCTEXT_NAMESPACE
diff --git a/Engine/Plugins/Experimental/ChaosEditor/Source/FractureEditor/Private/FractureToolBitmap.h b/Engine/Plugins/Experimental/ChaosEditor/Source/FractureEditor/Private/FractureToolBitmap.h
--- a/Engine/Plugins/Experimental/ChaosEditor/Source/FractureEditor/Private/FractureToolBitmap.h
+++ b/Engine/Plugins/Experimental/ChaosEditor/Source/FractureEditor/Private/FractureToolBitmap.h
@@ -80,5 +80,6 @@ public:
 	virtual TArray<UObject*> GetSettingsObjects() const override;// { return TArray<UObject*>(); }
 	// virtual void ExecuteFracture() {}
 	// virtual bool CanExecuteFracture() { return true; }
+	virtual bool CanExecuteFracture() const override;
 
 };
\ No newline at end of file
diff --git a/Engine/Plugins/Experimental/ChaosEditor/Source/FractureEditor/Private/FractureToolBrick.cpp b/Engine/Plugins/Experimental/ChaosEditor/Source/FractureEditor/Private/FractureToolBrick.cpp
--- a/Engine/Plugins/Experimental/ChaosEditor/Source/FractureEditor/Private/FractureToolBrick.cpp
+++ b/Engine/Plugins/Experimental/ChaosEditor/Source/FractureEditor/Private/FractureToolBrick.cpp
@@ -2,6 +2,7 @@
 
 #include "FractureToolBrick.h"
 
+#include "FractureEditorModeToolkit.h"
 #include "FractureEditorStyle.h"
 #include "Editor.h"
 #include "Engine/Selection.h"
@@ -63,6 +64,7 @@ void UFractureToolBrick::RegisterUICommand( FFractureEditorCommands* BindingCont
 TArray<UObject*> UFractureToolBrick::GetSettingsObjects() const 
 { 
 	TArray<UObject*> Settings; 
+	Settings.Add(GetMutableDefault<UFractureCommonSettings>());
 	Settings.Add(GetMutableDefault<UFractureBrickSettings>());
 	return Settings;
 }
@@ -364,4 +366,9 @@ void UFractureToolBrick::ExecuteFracture(const FFractureContext& FractureContext
 	}
 }
 
+bool UFractureToolBrick::CanExecuteFracture() const
+{
+	return FFractureEditorModeToolkit::IsLeafBoneSelected();
+}
+
 #undef LOCTEXT_NAMESPACE
diff --git a/Engine/Plugins/Experimental/ChaosEditor/Source/FractureEditor/Private/FractureToolBrick.h b/Engine/Plugins/Experimental/ChaosEditor/Source/FractureEditor/Private/FractureToolBrick.h
--- a/Engine/Plugins/Experimental/ChaosEditor/Source/FractureEditor/Private/FractureToolBrick.h
+++ b/Engine/Plugins/Experimental/ChaosEditor/Source/FractureEditor/Private/FractureToolBrick.h
@@ -111,6 +111,8 @@ public:
 	virtual void Render(const FSceneView* View, FViewport* Viewport, FPrimitiveDrawInterface* PDI) override;
 
 	virtual void ExecuteFracture(const FFractureContext& FractureContext) override;
+	virtual bool CanExecuteFracture() const override;
+
 private:
 	UPROPERTY()
 	UFractureBrickSettings* BrickSettings;
diff --git a/Engine/Plugins/Experimental/ChaosEditor/Source/FractureEditor/Private/FractureToolCluster.cpp b/Engine/Plugins/Experimental/ChaosEditor/Source/FractureEditor/Private/FractureToolCluster.cpp
--- a/Engine/Plugins/Experimental/ChaosEditor/Source/FractureEditor/Private/FractureToolCluster.cpp
+++ b/Engine/Plugins/Experimental/ChaosEditor/Source/FractureEditor/Private/FractureToolCluster.cpp
@@ -58,6 +58,7 @@ void UFractureToolCluster::RegisterUICommand( FFractureEditorCommands* BindingCo
 TArray<UObject*> UFractureToolCluster::GetSettingsObjects() const
 {
 	TArray<UObject*> ReturnSettings;
+	ReturnSettings.Add(GetMutableDefault<UFractureCommonSettings>());
 	ReturnSettings.Add(GetMutableDefault<UFractureClusterSettings>());
 	return ReturnSettings;
 }
diff --git a/Engine/Plugins/Experimental/ChaosEditor/Source/FractureEditor/Private/FractureToolPlaneCut.cpp b/Engine/Plugins/Experimental/ChaosEditor/Source/FractureEditor/Private/FractureToolPlaneCut.cpp
--- a/Engine/Plugins/Experimental/ChaosEditor/Source/FractureEditor/Private/FractureToolPlaneCut.cpp
+++ b/Engine/Plugins/Experimental/ChaosEditor/Source/FractureEditor/Private/FractureToolPlaneCut.cpp
@@ -91,6 +91,7 @@ void UFractureToolPlaneCut::Render(const FSceneView* View, FViewport* Viewport,
 TArray<UObject*> UFractureToolPlaneCut::GetSettingsObjects() const
 {
 	TArray<UObject*> Settings;
+	Settings.Add(GetMutableDefault<UFractureCommonSettings>());
 	Settings.Add(GetMutableDefault<UFracturePlaneCutSettings>());
 	return Settings;
 }
@@ -161,6 +162,11 @@ void UFractureToolPlaneCut::ExecuteFracture(const FFractureContext& FractureCont
 	}
 }
 
+bool UFractureToolPlaneCut::CanExecuteFracture() const
+{
+	return FFractureEditorModeToolkit::IsLeafBoneSelected();
+}
+
 #if WITH_EDITOR
 void UFractureToolPlaneCut::PostEditChangeChainProperty(struct FPropertyChangedChainEvent& PropertyChangedEvent)
 {
diff --git a/Engine/Plugins/Experimental/ChaosEditor/Source/FractureEditor/Private/FractureToolPlaneCut.h b/Engine/Plugins/Experimental/ChaosEditor/Source/FractureEditor/Private/FractureToolPlaneCut.h
--- a/Engine/Plugins/Experimental/ChaosEditor/Source/FractureEditor/Private/FractureToolPlaneCut.h
+++ b/Engine/Plugins/Experimental/ChaosEditor/Source/FractureEditor/Private/FractureToolPlaneCut.h
@@ -54,6 +54,7 @@ public:
 
 	virtual void FractureContextChanged() override;
 	virtual void ExecuteFracture(const FFractureContext& FractureContext) override;
+	virtual bool CanExecuteFracture() const override;
 
 #if WITH_EDITOR
 	virtual void PostEditChangeChainProperty(struct FPropertyChangedChainEvent& PropertyChangedEvent) override;
diff --git a/Engine/Plugins/Experimental/ChaosEditor/Source/FractureEditor/Private/FractureToolRadial.cpp b/Engine/Plugins/Experimental/ChaosEditor/Source/FractureEditor/Private/FractureToolRadial.cpp
--- a/Engine/Plugins/Experimental/ChaosEditor/Source/FractureEditor/Private/FractureToolRadial.cpp
+++ b/Engine/Plugins/Experimental/ChaosEditor/Source/FractureEditor/Private/FractureToolRadial.cpp
@@ -56,6 +56,7 @@ void UFractureToolRadial::RegisterUICommand( FFractureEditorCommands* BindingCon
 TArray<UObject*> UFractureToolRadial::GetSettingsObjects() const 
 { 
 	TArray<UObject*> Settings; 
+	Settings.Add(GetMutableDefault<UFractureCommonSettings>());
 	Settings.Add(GetMutableDefault<UFractureRadialSettings>());
 	return Settings;
 }
diff --git a/Engine/Plugins/Experimental/ChaosEditor/Source/FractureEditor/Private/FractureToolSlice.cpp b/Engine/Plugins/Experimental/ChaosEditor/Source/FractureEditor/Private/FractureToolSlice.cpp
--- a/Engine/Plugins/Experimental/ChaosEditor/Source/FractureEditor/Private/FractureToolSlice.cpp
+++ b/Engine/Plugins/Experimental/ChaosEditor/Source/FractureEditor/Private/FractureToolSlice.cpp
@@ -63,6 +63,7 @@ void UFractureToolSlice::RegisterUICommand( FFractureEditorCommands* BindingCont
 TArray<UObject*> UFractureToolSlice::GetSettingsObjects() const 
 { 
 	TArray<UObject*> Settings; 
+	Settings.Add(GetMutableDefault<UFractureCommonSettings>());
 	Settings.Add(GetMutableDefault<UFractureSliceSettings>());
 	return Settings;
 }
@@ -219,4 +220,9 @@ void UFractureToolSlice::ExecuteFracture(const FFractureContext& FractureContext
 	}
 }
 
+bool UFractureToolSlice::CanExecuteFracture() const
+{
+	return FFractureEditorModeToolkit::IsLeafBoneSelected();
+}
+
 #undef LOCTEXT_NAMESPACE
\ No newline at end of file
diff --git a/Engine/Plugins/Experimental/ChaosEditor/Source/FractureEditor/Private/FractureToolSlice.h b/Engine/Plugins/Experimental/ChaosEditor/Source/FractureEditor/Private/FractureToolSlice.h
--- a/Engine/Plugins/Experimental/ChaosEditor/Source/FractureEditor/Private/FractureToolSlice.h
+++ b/Engine/Plugins/Experimental/ChaosEditor/Source/FractureEditor/Private/FractureToolSlice.h
@@ -71,7 +71,7 @@ public:
 
 	virtual void FractureContextChanged() override;
 	virtual void ExecuteFracture(const FFractureContext& Context) override;
-	// virtual bool CanExecuteFracture() { return true; }
+	virtual bool CanExecuteFracture() const override;
 
 	void GenerateSliceTransforms(const FFractureContext& Context, TArray<FTransform>& CuttingPlaneTransforms);
 
diff --git a/Engine/Plugins/Experimental/ChaosEditor/Source/FractureEditor/Private/FractureToolUniform.cpp b/Engine/Plugins/Experimental/ChaosEditor/Source/FractureEditor/Private/FractureToolUniform.cpp
--- a/Engine/Plugins/Experimental/ChaosEditor/Source/FractureEditor/Private/FractureToolUniform.cpp
+++ b/Engine/Plugins/Experimental/ChaosEditor/Source/FractureEditor/Private/FractureToolUniform.cpp
@@ -57,6 +57,7 @@ void UFractureToolUniform::RegisterUICommand( FFractureEditorCommands* BindingCo
 TArray<UObject*> UFractureToolUniform::GetSettingsObjects() const 
 { 
 	TArray<UObject*> AllSettings; 
+	AllSettings.Add(GetMutableDefault<UFractureCommonSettings>());
 	AllSettings.Add(Settings);
 	return AllSettings;
 }
diff --git a/Engine/Plugins/Experimental/ChaosEditor/Source/FractureEditor/Private/FractureToolVoronoiBase.cpp b/Engine/Plugins/Experimental/ChaosEditor/Source/FractureEditor/Private/FractureToolVoronoiBase.cpp
--- a/Engine/Plugins/Experimental/ChaosEditor/Source/FractureEditor/Private/FractureToolVoronoiBase.cpp
+++ b/Engine/Plugins/Experimental/ChaosEditor/Source/FractureEditor/Private/FractureToolVoronoiBase.cpp
@@ -110,3 +110,8 @@ void UFractureToolVoronoiBase::ExecuteFracture(const FFractureContext& FractureC
 		}
 	}
 }
+
+bool UFractureToolVoronoiBase::CanExecuteFracture() const
+{
+	return FFractureEditorModeToolkit::IsLeafBoneSelected();
+}
diff --git a/Engine/Plugins/Experimental/ChaosEditor/Source/FractureEditor/Private/FractureToolVoronoiBase.h b/Engine/Plugins/Experimental/ChaosEditor/Source/FractureEditor/Private/FractureToolVoronoiBase.h
--- a/Engine/Plugins/Experimental/ChaosEditor/Source/FractureEditor/Private/FractureToolVoronoiBase.h
+++ b/Engine/Plugins/Experimental/ChaosEditor/Source/FractureEditor/Private/FractureToolVoronoiBase.h
@@ -23,6 +23,7 @@ public:
 
 	virtual void FractureContextChanged() override;
 	virtual void ExecuteFracture(const FFractureContext& FractureContext) override;
+	virtual bool CanExecuteFracture() const override;
 
 protected:
 	
diff --git a/Engine/Plugins/Experimental/ChaosEditor/Source/FractureEditor/Public/FractureEditorModeToolkit.h b/Engine/Plugins/Experimental/ChaosEditor/Source/FractureEditor/Public/FractureEditorModeToolkit.h
--- a/Engine/Plugins/Experimental/ChaosEditor/Source/FractureEditor/Public/FractureEditorModeToolkit.h
+++ b/Engine/Plugins/Experimental/ChaosEditor/Source/FractureEditor/Public/FractureEditorModeToolkit.h
@@ -8,7 +8,6 @@
 #include "Widgets/Input/SCheckBox.h"
 #include "Widgets/Layout/SSplitter.h"
 #include "GeometryCollection/GeometryCollectionActor.h"
-#include "AutoClusterFracture.h"
 
 class IDetailsView;
 class SScrollBox;
@@ -59,7 +58,6 @@ public:
 
 
 	// Cluster Command Callbacks
-	void OnAutoCluster();
 	void OnCluster();
 	void OnUncluster();
 	void OnFlatten();
@@ -74,7 +72,8 @@ public:
 	// Fracture Command Callback
 	FReply OnFractureClicked();
 	bool CanExecuteFracture() const;
-	bool IsLeafBoneSelected() const;
+
+	static bool IsLeafBoneSelected();
 
 	static ULevel* GetSelectedLevel();
  	static AActor* AddActor(ULevel* InLevel, UClass* Class);
@@ -97,13 +96,6 @@ public:
 
 	FText GetStatisticsSummary() const;
 
-	TSharedRef<SWidget> GetAutoClusterModesMenu();	
-	void SetAutoClusterMode(EFractureAutoClusterMode InAutoClusterMode);
-	FText GetAutoClusterModeDisplayName() const;
-
-	void SetAutoClusterSiteCount(uint32 InSiteCount);
-	uint32 GetAutoClusterSiteCount() const;
-
 	/** Returns the number of Mode specific tabs in the mode toolbar **/ 
 	const static TArray<FName> PaletteNames;
 	virtual void GetToolPaletteNames( TArray<FName>& InPaletteName ) const { InPaletteName = PaletteNames; }
@@ -118,7 +110,7 @@ public:
 protected:
 	static bool IsGeometryCollectionSelected();
 	static bool IsStaticMeshSelected();
-	static bool IsSeletedActorsInEditorWorld();
+	static bool IsSelectedActorsInEditorWorld();
 	void UpdateExplodedVectors(UGeometryCollectionComponent* GeometryCollectionComponent) const;
 private:
 	void OnOutlinerBoneSelectionChanged(UGeometryCollectionComponent* RootComponent, TArray<int32>& SelectedBones);
@@ -127,9 +119,6 @@ private:
 	float ExplodeAmount;
 	int32 FractureLevel;
 
-	EFractureAutoClusterMode AutoClusterMode;
-	uint32 AutoClusterSiteCount;
-
 	UFractureTool* ActiveTool;
 
 	TSharedPtr<IDetailsView> DetailsView;
diff --git a/Engine/Plugins/Experimental/ChaosEditor/Source/FractureEditor/Public/FractureTool.h b/Engine/Plugins/Experimental/ChaosEditor/Source/FractureEditor/Public/FractureTool.h
--- a/Engine/Plugins/Experimental/ChaosEditor/Source/FractureEditor/Public/FractureTool.h
+++ b/Engine/Plugins/Experimental/ChaosEditor/Source/FractureEditor/Public/FractureTool.h
@@ -93,9 +93,13 @@ public:
 
 	UFractureTool(const FObjectInitializer& ObjInit) : Super(ObjInit) {}
 
+	/** This is the Text that will appear on the tool button to execute the fracture **/
 	virtual FText GetDisplayText() const  { return FText(); }
 	virtual FText GetTooltipText() const { return FText(); }
 
+	/** This is the Text that will appear on the button to execute the fracture **/
+	virtual FText GetApplyText() const { return FText(NSLOCTEXT("Fracture", "ExecuteFracture", "Fracture")); }
+
 	// TODO: What's the correct thing to return here?
 	virtual FSlateIcon GetToolIcon() const { return FSlateIcon(); }
 
