Fix warning C4005: 'LOCTEXT_NAMESPACE': macro redefinition
#rnx

diff --git a/Engine/Plugins/Experimental/ChaosEditor/Source/FractureEditor/Private/FractureToolAutoCluster.cpp b/Engine/Plugins/Experimental/ChaosEditor/Source/FractureEditor/Private/FractureToolAutoCluster.cpp
--- a/Engine/Plugins/Experimental/ChaosEditor/Source/FractureEditor/Private/FractureToolAutoCluster.cpp
+++ b/Engine/Plugins/Experimental/ChaosEditor/Source/FractureEditor/Private/FractureToolAutoCluster.cpp
@@ -70,5 +70,5 @@ void UFractureToolAutoCluster::ExecuteFracture(const FFractureContext& FractureC
 	}
 }
 
-
+#undef LOCTEXT_NAMESPACE
 
