Fix editor usage of root object details customization.  No longer needed with details panel changes

diff --git a/Engine/Plugins/Experimental/ChaosEditor/Source/FractureEditor/Private/FractureEditorModeToolkit.cpp b/Engine/Plugins/Experimental/ChaosEditor/Source/FractureEditor/Private/FractureEditorModeToolkit.cpp
--- a/Engine/Plugins/Experimental/ChaosEditor/Source/FractureEditor/Private/FractureEditorModeToolkit.cpp
+++ b/Engine/Plugins/Experimental/ChaosEditor/Source/FractureEditor/Private/FractureEditorModeToolkit.cpp
@@ -52,6 +52,7 @@
 #include "GeometryCollection/GeometryCollectionClusteringUtility.h"
 #include "Editor.h"
 #include "LevelEditorViewport.h"
+#include "Widgets/Layout/SUniformGridPanel.h"
 
 #define LOCTEXT_NAMESPACE "FFractureEditorModeToolkit"
 
@@ -112,10 +113,9 @@ void FFractureEditorModeToolkit::Init(const TSharedPtr<IToolkitHost>& InitToolki
 	DetailsViewArgs.bShowAnimatedPropertiesOption = false;
 
 	DetailsView = EditModule.CreateDetailView(DetailsViewArgs);
-	DetailsView->SetRootObjectCustomizationInstance(MakeShareable(new FFractureRootObjectCustomization));
 
 	float Padding = 4.0f;
-	float MorePadding = 10.0f;
+	FMargin MorePadding = FMargin(10.0f, 2.0f);
 
 	SAssignNew(ExplodedViewWidget, SSpinBox<int32>)
 	.Style(&FFractureEditorStyle::Get().GetWidgetStyle<FSpinBoxStyle>("FractureEditor.SpinBox"))
@@ -177,7 +177,7 @@ void FFractureEditorModeToolkit::Init(const TSharedPtr<IToolkitHost>& InitToolki
 	TSharedRef<SExpandableArea> OutlinerExpander = SNew(SExpandableArea)
 	.AreaTitle(FText(LOCTEXT("Outliner", "Outliner")))
 	.HeaderPadding(FMargin(2.0, 2.0))
-	.Padding(FMargin(MorePadding))
+	.Padding(MorePadding)
 	.BorderImage(FEditorStyle::Get().GetBrush("DetailsView.CategoryTop"))
 	.BorderBackgroundColor( FLinearColor( .6,.6,.6, 1.0f ) )
 	.BodyBorderBackgroundColor (FLinearColor( 1.0, 0.0, 0.0))
@@ -191,7 +191,7 @@ void FFractureEditorModeToolkit::Init(const TSharedPtr<IToolkitHost>& InitToolki
 	TSharedRef<SExpandableArea> StatisticsExpander = SNew(SExpandableArea)
 	.AreaTitle(FText(LOCTEXT("LevelStatistics", "Level Statistics")))
 	.HeaderPadding(FMargin(2.0, 2.0))
-	.Padding(FMargin(MorePadding))
+	.Padding(MorePadding)
 	.BorderImage(FEditorStyle::Get().GetBrush("DetailsView.CategoryTop"))
 	.BorderBackgroundColor( FLinearColor( .6,.6,.6, 1.0f ) )
 	.BodyBorderBackgroundColor (FLinearColor( 1.0, 0.0, 0.0))
@@ -203,7 +203,6 @@ void FFractureEditorModeToolkit::Init(const TSharedPtr<IToolkitHost>& InitToolki
 	];
 
 	SAssignNew(ToolkitWidget, SBox)
-	.Padding(8)
 	[
 		SNew(SVerticalBox)
 
@@ -235,53 +234,37 @@ void FFractureEditorModeToolkit::Init(const TSharedPtr<IToolkitHost>& InitToolki
 					StatisticsExpander
 				]
 			]
-
 			+ SWidgetSwitcher::Slot()
 			[
-
 				SNew(SScrollBox)
-
 				+SScrollBox::Slot()
-				.Padding(0.0f, MorePadding)
+				.Padding(0.0f, 10.0f, 0.0f, 0.0f)
 				[
 					DetailsView.ToSharedRef()
 				]
-
 				+SScrollBox::Slot()
-				.Padding(16.f)
+				.HAlign(HAlign_Center)
 				[
-					SNew(SHorizontalBox)
-
-					+SHorizontalBox::Slot()
-					.FillWidth(1)
-
-					+SHorizontalBox::Slot()
-					.AutoWidth()
+					SNew(SUniformGridPanel)
+					.MinDesiredSlotWidth(100.f)
+					.SlotPadding(10.f)
+					+SUniformGridPanel::Slot(0, 0)
 					[
 						SNew( SButton )
 						.HAlign(HAlign_Center)
-						.ContentPadding(FMargin(MorePadding, Padding))
+						.ContentPadding(FMargin(10.f, Padding))
 						.OnClicked(this, &FFractureEditorModeToolkit::OnFractureClicked)
 						.IsEnabled( this, &FFractureEditorModeToolkit::CanExecuteFracture)
 						.Text_Lambda( [this] () -> FText { return ActiveTool ? ActiveTool->GetApplyText() :  LOCTEXT("FractureApplyButton", "Apply"); })
 					]
-
-					+SHorizontalBox::Slot()
-					.FillWidth(1)
-
-					+SHorizontalBox::Slot()
-					.AutoWidth()
+					+ SUniformGridPanel::Slot(1, 0)
 					[
 						SNew(SButton)
 						.HAlign(HAlign_Center)
-						.ContentPadding(FMargin(MorePadding, Padding))
+						.ContentPadding(FMargin(10.f, Padding))
 						.OnClicked_Lambda( [this] () -> FReply { SetActiveTool(0); return FReply::Handled(); } )
 						.Text(FText(LOCTEXT("FractureCancelButton", "Cancel")))
 					]
-
-
-					+SHorizontalBox::Slot()
-					.FillWidth(1)
 				]
 			]
 		]
