Moved Modes Tools into Modes Panel, removed ToolBar

diff --git a/Engine/Plugins/Experimental/ChaosEditor/Source/FractureEditor/Private/FractureEditorModeToolkit.cpp b/Engine/Plugins/Experimental/ChaosEditor/Source/FractureEditor/Private/FractureEditorModeToolkit.cpp
--- a/Engine/Plugins/Experimental/ChaosEditor/Source/FractureEditor/Private/FractureEditorModeToolkit.cpp
+++ b/Engine/Plugins/Experimental/ChaosEditor/Source/FractureEditor/Private/FractureEditorModeToolkit.cpp
@@ -34,6 +34,7 @@
 #include "GeometryCollection/GeometryCollectionActor.h"
 #include "GeometryCollection/GeometryCollection.h"
 
+#include "Styling/AppStyle.h"
 #include "EditorStyleSet.h"
 #include "FractureEditor.h"
 #include "FractureEditorCommands.h"
@@ -54,6 +55,13 @@
 #include "LevelEditorViewport.h"
 #include "Widgets/Layout/SUniformGridPanel.h"
 
+#include "DetailWidgetRow.h"
+#include "DetailCategoryBuilder.h"
+#include "DetailLayoutBuilder.h"
+#include "PropertyHandle.h"
+
+#include "FractureSettings.h"
+
 #define LOCTEXT_NAMESPACE "FFractureEditorModeToolkit"
 
 TArray<UClass*> FindFractureToolClasses()
@@ -71,13 +79,90 @@ TArray<UClass*> FindFractureToolClasses()
 	return Classes;
 }
 
+FFractureViewSettingsCustomization::FFractureViewSettingsCustomization(FFractureEditorModeToolkit* InToolkit) 
+	: Toolkit(InToolkit)
+{
+
+}
+
+TSharedRef<IDetailCustomization> FFractureViewSettingsCustomization::MakeInstance(FFractureEditorModeToolkit* InToolkit)
+{
+	return MakeShareable(new FFractureViewSettingsCustomization(InToolkit));
+}
+
+void 
+FFractureViewSettingsCustomization::CustomizeDetails(IDetailLayoutBuilder& DetailBuilder) 
+{
+	IDetailCategoryBuilder& ViewCategory = DetailBuilder.EditCategory("ViewSettings", FText::GetEmpty(), ECategoryPriority::TypeSpecific);
+
+	TSharedRef<IPropertyHandle> LevelProperty = DetailBuilder.GetProperty("FractureLevel");
+
+	ViewCategory.AddCustomRow(FText::GetEmpty())
+	.NameContent()
+	.HAlign(HAlign_Right)
+	[
+		SNew(STextBlock)
+		.TextStyle( &FCoreStyle::Get().GetWidgetStyle<FTextBlockStyle>( "SmallText" ) )
+		.Text(LOCTEXT("ShowBoneColors", "Show Bone Colors"))
+
+	]
+	.ValueContent()
+	[
+		SNew(SCheckBox)
+		.IsChecked_Lambda([this]() -> ECheckBoxState { return this->Toolkit->GetShowBoneColors() ? ECheckBoxState::Checked : ECheckBoxState::Unchecked; } )
+		.OnCheckStateChanged_Lambda([this] (ECheckBoxState CheckState) { this->Toolkit->OnSetShowBoneColors( CheckState == ECheckBoxState::Checked ); } )
+	];
+
+
+	ViewCategory.AddProperty(LevelProperty)
+	.CustomWidget()
+	.NameContent()
+	.HAlign(HAlign_Right)
+	[
+		SNew(STextBlock)
+		.TextStyle( &FCoreStyle::Get().GetWidgetStyle<FTextBlockStyle>( "SmallText" ) )
+		.Text(LevelProperty->GetPropertyDisplayName())
+	]
+	.ValueContent()
+	[
+		SNew( SComboButton)
+		.ContentPadding(0)
+		.OnGetMenuContent(Toolkit, &FFractureEditorModeToolkit::GetLevelViewMenuContent, LevelProperty) 
+		.ButtonContent()
+		[
+			SNew(STextBlock)
+			.Justification(ETextJustify::Left)
+			.Text_Lambda( [=]() -> FText {
+
+				int32 FractureLevel = 5;
+				LevelProperty->GetValue(FractureLevel);
+
+				if (FractureLevel < 0)
+				{
+					return LOCTEXT("FractureViewAllLevels", "All");
+				}
+				else if (FractureLevel == 0)
+				{
+					return LOCTEXT("FractureViewRootLevel", "Root");
+				}
+
+				return FText::Format(NSLOCTEXT("FractureEditor", "CurrentLevel", "{0}"), FText::AsNumber(FractureLevel));
+
+			})
+		]
+	];
+};
+
 FFractureEditorModeToolkit::FFractureEditorModeToolkit()
-	: ExplodeAmount(0.0f)
-	, FractureLevel(-1)
-	, ActiveTool(nullptr)
+	: ActiveTool(nullptr)
 {
 }
 
+FFractureEditorModeToolkit::~FFractureEditorModeToolkit()
+{
+	FCoreUObjectDelegates::OnObjectPropertyChanged.RemoveAll(this);
+}
+
 void FFractureEditorModeToolkit::Init(const TSharedPtr<IToolkitHost>& InitToolkitHost)
 {
 	FFractureEditorModule& FractureModule = FModuleManager::GetModuleChecked<FFractureEditorModule>("FractureEditor");
@@ -86,6 +171,8 @@ void FFractureEditorModeToolkit::Init(const TSharedPtr<IToolkitHost>& InitToolki
 
 	FPropertyEditorModule& EditModule = FModuleManager::Get().GetModuleChecked<FPropertyEditorModule>("PropertyEditor");
 
+	FCoreUObjectDelegates::OnObjectPropertyChanged.AddRaw(this, &FFractureEditorModeToolkit::OnObjectPostEditChange);
+
 	FDetailsViewArgs DetailsViewArgs(
 		false,  //bUpdateFromSelection=
 		false, //bLockable=
@@ -104,67 +191,15 @@ void FFractureEditorModeToolkit::Init(const TSharedPtr<IToolkitHost>& InitToolki
 	DetailsViewArgs.bShowAnimatedPropertiesOption = false;
 
 	DetailsView = EditModule.CreateDetailView(DetailsViewArgs);
+	EditModule.RegisterCustomClassLayout("FractureSettings", FOnGetDetailCustomizationInstance::CreateStatic(&FFractureViewSettingsCustomization::MakeInstance, this));
+
+	TArray<UObject*> Settings;
+	Settings.Add(GetMutableDefault<UFractureSettings>());
+	DetailsView->SetObjects(Settings);
 
 	float Padding = 4.0f;
 	FMargin MorePadding = FMargin(10.0f, 2.0f);
 
-	SAssignNew(ExplodedViewWidget, SSpinBox<int32>)
-	.Style(&FFractureEditorStyle::Get().GetWidgetStyle<FSpinBoxStyle>("FractureEditor.SpinBox"))
-	.PreventThrottling(true)
-	.Value_Lambda([=]() -> int32 { return (int32) (ExplodeAmount * 100.0f); })
-	.OnValueChanged_Lambda( [=](int32 NewValue) { this->OnSetExplodedViewValue( (float)NewValue / 100.0f); } )
-	.MinValue(0)
-	.MaxValue(100)
-	.Delta(1)
-	.Font(FCoreStyle::GetDefaultFontStyle("Regular", 11))
-	.MinDesiredWidth(36.f)
-	.Justification(ETextJustify::Center)
-	.ToolTipText(LOCTEXT("FractureEditor.Exploded_Tooltip", "How much to seperate the drawing of the bones to aid in setup.  Does not effect simulation"))
-	;
-
-	SAssignNew(LevelViewWidget, SComboButton)
-	.ContentPadding(0)
-	.ButtonStyle(FEditorStyle::Get(), "Toolbar.Button")
-	.ForegroundColor(FEditorStyle::Get().GetSlateColor("ToolBar.SToolBarComboButtonBlock.ComboButton.Color"))
-	.OnGetMenuContent(this, &FFractureEditorModeToolkit::GetLevelViewMenuContent) 
-	.ButtonContent()
-	[
-		SNew(SBox)
-		.WidthOverride(36)
-		[
-			SNew(STextBlock)
-			.Justification(ETextJustify::Center)
-			.Text_Lambda( [=]() -> FText {
-				if (FractureLevel < 0)
-				{
-					return LOCTEXT("FractureViewAllLevels", "All");
-				}
-				else if (FractureLevel == 0)
-				{
-					return LOCTEXT("FractureViewRootLevel", "Root");
-				}
-
-				return FText::Format(NSLOCTEXT("FractureEditor", "CurrentLevel", "{0}"), FText::AsNumber(FractureLevel));
-
-			})
-			.Font(FCoreStyle::GetDefaultFontStyle("Regular", 11))
-			.ToolTipText(LOCTEXT("FractureEditor.Level_Tooltip", "Set the currently view level of the geometry collection"))
-		]
-
-	];
-
-	SAssignNew(ShowBoneColorsWidget, SComboButton)
-	.ContentPadding(0)
-	.ButtonStyle(FEditorStyle::Get(), "ToolBar.Button")
-	.ForegroundColor(FEditorStyle::Get().GetSlateColor("ToolBar.SToolBarComboButtonBlock.ComboButton.Color"))
-	.OnGetMenuContent(this, &FFractureEditorModeToolkit::GetViewMenuContent)
-	.ButtonContent()
-	[
-		SNew(SImage)
-		.Image(FFractureEditorStyle::Get().GetBrush("FractureEditor.Visibility"))
-		.ToolTipText(LOCTEXT("FractureEditor.Visibility_Tooltip", "Toggle showing the bone colours of the geometry collection"))
-	];
-
 	TSharedRef<SExpandableArea> OutlinerExpander = SNew(SExpandableArea)
 	.AreaTitle(FText(LOCTEXT("Outliner", "Outliner")))
 	.HeaderPadding(FMargin(2.0, 2.0))
@@ -173,6 +208,7 @@ void FFractureEditorModeToolkit::Init(const TSharedPtr<IToolkitHost>& InitToolki
 	.BorderBackgroundColor( FLinearColor( .6,.6,.6, 1.0f ) )
 	.BodyBorderBackgroundColor (FLinearColor( 1.0, 0.0, 0.0))
 	.AreaTitleFont(FEditorStyle::Get().GetFontStyle("DetailsView.CategoryFontStyle"))
+	.Visibility_Lambda([this]() -> EVisibility { return (GetActiveTool() == nullptr) ? EVisibility::Visible: EVisibility::Collapsed; })
 	.BodyContent()
 	[
 		SAssignNew(OutlinerView, SGeometryCollectionOutliner )
@@ -187,6 +223,7 @@ void FFractureEditorModeToolkit::Init(const TSharedPtr<IToolkitHost>& InitToolki
 	.BorderBackgroundColor( FLinearColor( .6,.6,.6, 1.0f ) )
 	.BodyBorderBackgroundColor (FLinearColor( 1.0, 0.0, 0.0))
 	.AreaTitleFont(FEditorStyle::Get().GetFontStyle("DetailsView.CategoryFontStyle"))
+	.Visibility_Lambda([this]() -> EVisibility { return (GetActiveTool() == nullptr) ? EVisibility::Visible: EVisibility::Collapsed; })
 	.BodyContent()
 	[
 		SNew(STextBlock)
@@ -199,65 +236,82 @@ void FFractureEditorModeToolkit::Init(const TSharedPtr<IToolkitHost>& InitToolki
 
 		+SVerticalBox::Slot()
 		[
-			SNew(SWidgetSwitcher)
-			.WidgetIndex_Lambda( [this] () -> uint32 { return (GetActiveTool() != nullptr ? 1 : 0); })
-
-			+ SWidgetSwitcher::Slot()
-			[
-				SNew(SSplitter)
-				.Orientation( Orient_Vertical )
-
-				+SSplitter::Slot()
-				.SizeRule( TAttribute<SSplitter::ESizeRule>::Create( [this, OutlinerExpander] () { 
-					return OutlinerExpander->IsExpanded() ? SSplitter::ESizeRule::FractionOfParent : SSplitter::ESizeRule::SizeToContent; 
-				} ) )
-				.Value(1.f)
-				[
-					OutlinerExpander
-				]
 
-				+SSplitter::Slot()
-				.SizeRule( TAttribute<SSplitter::ESizeRule>::Create( [this, StatisticsExpander] () { 
-					return StatisticsExpander->IsExpanded() ? SSplitter::ESizeRule::FractionOfParent : SSplitter::ESizeRule::SizeToContent; 
-				} ) )
-				.Value(0.25f)
-				[
-					StatisticsExpander
-				]
-			]
-			+ SWidgetSwitcher::Slot()
+			SNew(SSplitter)
+			.Orientation( Orient_Vertical )
+			+SSplitter::Slot()
+			.SizeRule( TAttribute<SSplitter::ESizeRule>::Create( [this] () { 
+				return (GetActiveTool() != nullptr) ? SSplitter::ESizeRule::FractionOfParent : SSplitter::ESizeRule::SizeToContent; 
+			} ) )
+			.Value(1.f)
 			[
-				SNew(SScrollBox)
-				+SScrollBox::Slot()
-				.Padding(0.0f, 10.0f, 0.0f, 0.0f)
+				SNew(SVerticalBox)
+				+SVerticalBox::Slot()
+				.FillHeight(1.0)
 				[
-					DetailsView.ToSharedRef()
+					SNew(SScrollBox)
+					+SScrollBox::Slot()
+					[
+						DetailsView.ToSharedRef()
+					]
 				]
-				+SScrollBox::Slot()
-				.HAlign(HAlign_Center)
+				+SVerticalBox::Slot()
+				.AutoHeight()
 				[
-					SNew(SUniformGridPanel)
-					.MinDesiredSlotWidth(100.f)
-					.SlotPadding(10.f)
-					+SUniformGridPanel::Slot(0, 0)
+					SNew(SHorizontalBox)
+					+SHorizontalBox::Slot()
+					[
+						SNew(SSpacer)	
+					]	
+
+					+SHorizontalBox::Slot()
+					.Padding(4.0)
+					.AutoWidth()
 					[
 						SNew( SButton )
+						.ButtonStyle(FAppStyle::Get(), "PrimaryButton")
+						.TextStyle( FAppStyle::Get(), "ButtonText" )
 						.HAlign(HAlign_Center)
 						.ContentPadding(FMargin(10.f, Padding))
 						.OnClicked(this, &FFractureEditorModeToolkit::OnFractureClicked)
 						.IsEnabled( this, &FFractureEditorModeToolkit::CanExecuteFracture)
-						.Text_Lambda( [this] () -> FText { return ActiveTool ? ActiveTool->GetApplyText() :  LOCTEXT("FractureApplyButton", "Apply"); })
+						.Text_Lambda( [this] () -> FText { return ActiveTool ? ActiveTool->GetApplyText().ToUpper() :  LOCTEXT("FractureApplyButton", "APPLY"); })
+						.Visibility_Lambda( [this] () -> EVisibility { return (GetActiveTool() == nullptr) ? EVisibility::Collapsed : EVisibility::Visible; })
 					]
-					+ SUniformGridPanel::Slot(1, 0)
+
+					+ SHorizontalBox::Slot()
+					.Padding(4.0)
+					.AutoWidth()
 					[
 						SNew(SButton)
+						.ButtonStyle(FAppStyle::Get(), "Button")
+						.TextStyle( FAppStyle::Get(), "ButtonText" )
 						.HAlign(HAlign_Center)
 						.ContentPadding(FMargin(10.f, Padding))
 						.OnClicked_Lambda( [this] () -> FReply { SetActiveTool(0); return FReply::Handled(); } )
-						.Text(FText(LOCTEXT("FractureCancelButton", "Cancel")))
+						.Text(FText(LOCTEXT("FractureCancelButton", "CANCEL")))
+						.Visibility_Lambda( [this] () -> EVisibility { return (GetActiveTool() == nullptr) ? EVisibility::Collapsed : EVisibility::Visible; })
 					]
 				]
 			]
+
+			+SSplitter::Slot()
+			.SizeRule( TAttribute<SSplitter::ESizeRule>::Create( [this, OutlinerExpander] () { 
+				return OutlinerExpander->IsExpanded() ? SSplitter::ESizeRule::FractionOfParent : SSplitter::ESizeRule::SizeToContent; 
+			} ) )
+			.Value(1.f)
+			[
+				OutlinerExpander
+			]
+
+			+SSplitter::Slot()
+			.SizeRule( TAttribute<SSplitter::ESizeRule>::Create( [this, StatisticsExpander] () { 
+				return StatisticsExpander->IsExpanded() ? SSplitter::ESizeRule::FractionOfParent : SSplitter::ESizeRule::SizeToContent; 
+			} ) )
+			.Value(0.25f)
+			[
+				StatisticsExpander
+			]
 		]
 	];
 
@@ -268,7 +322,22 @@ void FFractureEditorModeToolkit::Init(const TSharedPtr<IToolkitHost>& InitToolki
 
 }
 
-const TArray<FName> FFractureEditorModeToolkit::PaletteNames = { FName(TEXT("Fracture")), FName(TEXT("Cluster")) };
+void FFractureEditorModeToolkit::OnObjectPostEditChange( UObject* Object, FPropertyChangedEvent& PropertyChangedEvent )
+{
+	if (PropertyChangedEvent.Property)
+	{
+		if ( PropertyChangedEvent.Property->GetFName() == GET_MEMBER_NAME_CHECKED(UFractureSettings, ExplodeAmount))
+		{
+			OnExplodedViewValueChanged();
+		} 
+		else if ( PropertyChangedEvent.Property->GetFName() == GET_MEMBER_NAME_CHECKED(UFractureSettings, FractureLevel))
+		{
+			OnLevelViewValueChanged();
+		} 
+	}
+}
+
+const TArray<FName> FFractureEditorModeToolkit::PaletteNames = { FName(TEXT("Fracture")) };
 
 FText FFractureEditorModeToolkit::GetToolPaletteDisplayName(FName Palette) const
 { 
@@ -293,11 +362,6 @@ void FFractureEditorModeToolkit::BuildToolPalette(FName PaletteIndex, class FToo
 
 		ToolbarBuilder.AddSeparator();
 
-		ToolbarBuilder.AddToolBarWidget(ExplodedViewWidget.ToSharedRef(), FText(LOCTEXT("FractureExplodedPercentage", "Explode")));
-		ToolbarBuilder.AddToolBarWidget(LevelViewWidget.ToSharedRef(), LOCTEXT("FractureViewLevel", "Level"));
-		// ToolbarBuilder.AddToolBarWidget(ShowBoneColorsWidget.ToSharedRef(), LOCTEXT("FractureViewOptions", "View"));
-		ToolbarBuilder.AddToolBarButton(Commands.ToggleShowBoneColors);
-
 		ToolbarBuilder.AddSeparator();
 		ToolbarBuilder.AddToolBarButton(Commands.Uniform);
 		ToolbarBuilder.AddToolBarButton(Commands.Clustered);
@@ -305,44 +369,20 @@ void FFractureEditorModeToolkit::BuildToolPalette(FName PaletteIndex, class FToo
 		ToolbarBuilder.AddToolBarButton(Commands.Planar);
 		ToolbarBuilder.AddToolBarButton(Commands.Slice);
 		ToolbarBuilder.AddToolBarButton(Commands.Brick);
-	}	
-
-	else if (PaletteIndex == PaletteNames[1])
-	{
-		ToolbarBuilder.AddWidget(SNew(SBox).WidthOverride(4));
-		ToolbarBuilder.AddToolBarButton(Commands.GenerateAsset);
-		ToolbarBuilder.AddSeparator();
 
-		ToolbarBuilder.AddToolBarButton(Commands.SelectAll);
-		ToolbarBuilder.AddToolBarButton(Commands.SelectNone);
-		ToolbarBuilder.AddToolBarButton(Commands.SelectNeighbors);
-		ToolbarBuilder.AddToolBarButton(Commands.SelectSiblings);
-		ToolbarBuilder.AddToolBarButton(Commands.SelectAllInCluster);
-		ToolbarBuilder.AddToolBarButton(Commands.SelectInvert);
-
-		ToolbarBuilder.AddSeparator();
-
-		ToolbarBuilder.AddToolBarWidget(ExplodedViewWidget.ToSharedRef(), FText(LOCTEXT("FractureExplodedPercentage", "Explode")));
-		ToolbarBuilder.AddToolBarWidget(LevelViewWidget.ToSharedRef(), LOCTEXT("FractureViewLevel", "Level"));
-		// ToolbarBuilder.AddToolBarWidget(ShowBoneColorsWidget.ToSharedRef(), LOCTEXT("FractureViewOptions", "View"));
-		ToolbarBuilder.AddToolBarButton(Commands.ToggleShowBoneColors);
 
 		ToolbarBuilder.AddSeparator();
 
 		ToolbarBuilder.AddToolBarButton(Commands.AutoCluster);
-		ToolbarBuilder.AddSeparator();
-
 		ToolbarBuilder.AddToolBarButton(Commands.Flatten);
 		// ToolbarBuilder.AddToolBarButton(Commands.FlattenToLevel);
 		ToolbarBuilder.AddToolBarButton(Commands.Cluster);
 		ToolbarBuilder.AddToolBarButton(Commands.Uncluster);
 		// ToolbarBuilder.AddToolBarButton(Commands.Merge);
 		ToolbarBuilder.AddToolBarButton(Commands.MoveUp);
-
 	}
 }
 
-
 void FFractureEditorModeToolkit::BindCommands()
 {
 	const FFractureEditorCommands& Commands = FFractureEditorCommands::Get();
@@ -470,7 +510,7 @@ void FFractureEditorModeToolkit::BindCommands()
 
 	ToolkitCommands->MapAction(
 		Commands.ToggleShowBoneColors,
-		FExecuteAction::CreateSP(this, &FFractureEditorModeToolkit::OnSetShowBoneColors),
+		FExecuteAction::CreateSP(this, &FFractureEditorModeToolkit::OnToggleShowBoneColors),
 		FCanExecuteAction(),//::CreateLambda([]() { return true; }),
 		FIsActionChecked::CreateSP(this, &FFractureEditorModeToolkit::GetShowBoneColors)
 	);
@@ -487,14 +527,14 @@ void FFractureEditorModeToolkit::BindCommands()
 
 	ToolkitCommands->MapAction(
 		Commands.ExplodeMore,
-		FExecuteAction::CreateLambda([=]() { this->OnSetExplodedViewValue( FMath::Min(1.0, ExplodeAmount + .1) ); } ),
+		FExecuteAction::CreateLambda([=]() { this->OnSetExplodedViewValue( FMath::Min(1.0, this->GetExplodedViewValue() + .1) ); } ),
 		EUIActionRepeatMode::RepeatEnabled
 	);
 
 
 	ToolkitCommands->MapAction(
 		Commands.ExplodeLess,
-		FExecuteAction::CreateLambda([=]() { this->OnSetExplodedViewValue( FMath::Max(0.0, ExplodeAmount - .1) ); } ),
+		FExecuteAction::CreateLambda([=]() { this->OnSetExplodedViewValue( FMath::Max(0.0, this->GetExplodedViewValue() - .1) ); } ),
 		EUIActionRepeatMode::RepeatEnabled
 	);
 }
@@ -529,17 +569,19 @@ void FFractureEditorModeToolkit::AddReferencedObjects(FReferenceCollector& Colle
 
 float FFractureEditorModeToolkit::GetExplodedViewValue() const
 {
-	return ExplodeAmount;
+	UFractureSettings* FractureSettings = GetMutableDefault<UFractureSettings>();
+	return FractureSettings->ExplodeAmount;
 }
 
 int32 FFractureEditorModeToolkit::GetLevelViewValue() const
 {
-	return FractureLevel;
+	UFractureSettings* FractureSettings = GetMutableDefault<UFractureSettings>();
+	return FractureSettings->FractureLevel;
 }
 
-
 bool FFractureEditorModeToolkit::GetShowBoneColors() const
 {
+	
 	TSet<UGeometryCollectionComponent*> GeomCompSelection;
 	GetSelectedGeometryCollectionComponents(GeomCompSelection);
 	if (GeomCompSelection.Num() > 0)
@@ -552,57 +594,66 @@ bool FFractureEditorModeToolkit::GetShowBoneColors() const
 	return false;
 }
 
-void FFractureEditorModeToolkit::OnSetShowBoneColors()
+void FFractureEditorModeToolkit::OnToggleShowBoneColors()
 {
-	bool OldState = GetShowBoneColors();
+	OnSetShowBoneColors(!GetShowBoneColors());
+}
 
+void FFractureEditorModeToolkit::OnSetShowBoneColors(bool NewValue)
+{
 	TSet<UGeometryCollectionComponent*> GeomCompSelection;
 	GetSelectedGeometryCollectionComponents(GeomCompSelection);
 	for (UGeometryCollectionComponent* Comp : GeomCompSelection)
 	{
 		FScopedColorEdit EditBoneColor = Comp->EditBoneSelection();
-		EditBoneColor.SetShowBoneColors(!OldState);
+		EditBoneColor.SetShowBoneColors(NewValue);
 		Comp->MarkRenderStateDirty();
 		Comp->MarkRenderDynamicDataDirty();
 	}
 	GCurrentLevelEditingViewportClient->Invalidate();
 }
 
-
 void FFractureEditorModeToolkit::OnSetExplodedViewValue(float NewValue)
 {
-	if ( FMath::Abs<float>( ExplodeAmount - NewValue ) >= .01f)
+	FScopedTransaction Transaction(LOCTEXT("SetExplodedViewValue", "Adjust Exploded View"));
+
+	UFractureSettings* FractureSettings = GetMutableDefault<UFractureSettings>();
+	if ( FMath::Abs<float>( FractureSettings->ExplodeAmount - NewValue ) >= .01f)
 	{
-		ExplodeAmount = NewValue;
+		FractureSettings->ExplodeAmount = NewValue;
+		OnExplodedViewValueChanged();
+	}
+}
 
-		USelection* SelectionSet = GEditor->GetSelectedActors();
+void FFractureEditorModeToolkit::OnExplodedViewValueChanged()
+{
+	USelection* SelectionSet = GEditor->GetSelectedActors();
 
-		TArray<AActor*> SelectedActors;
-		SelectedActors.Reserve(SelectionSet->Num());
-		SelectionSet->GetSelectedObjects(SelectedActors);
+	TArray<AActor*> SelectedActors;
+	SelectedActors.Reserve(SelectionSet->Num());
+	SelectionSet->GetSelectedObjects(SelectedActors);
 
-		for (AActor* Actor : SelectedActors)
+	for (AActor* Actor : SelectedActors)
+	{
+		TInlineComponentArray<UPrimitiveComponent*> Components;
+		Actor->GetComponents(Components);
+		for (UPrimitiveComponent* PrimitiveComponent : Components)
 		{
-			TInlineComponentArray<UPrimitiveComponent*> Components;
-			Actor->GetComponents(Components);
-			for (UPrimitiveComponent* PrimitiveComponent : Components)
+			AGeometryCollectionActor* GeometryCollectionActor = Cast<AGeometryCollectionActor>(Actor);
+			if(GeometryCollectionActor)
 			{
-				AGeometryCollectionActor* GeometryCollectionActor = Cast<AGeometryCollectionActor>(Actor);
-				if(GeometryCollectionActor)
-				{
-					UGeometryCollectionComponent* GeometryCollectionComponent = Cast<UGeometryCollectionComponent>(PrimitiveComponent);
-					FGeometryCollectionEdit RestCollection = GeometryCollectionComponent->EditRestCollection();
-					UGeometryCollection* GeometryCollection = RestCollection.GetRestCollection();
+				UGeometryCollectionComponent* GeometryCollectionComponent = Cast<UGeometryCollectionComponent>(PrimitiveComponent);
+				FGeometryCollectionEdit RestCollection = GeometryCollectionComponent->EditRestCollection();
+				UGeometryCollection* GeometryCollection = RestCollection.GetRestCollection();
 
-					UpdateExplodedVectors(GeometryCollectionComponent);
+				UpdateExplodedVectors(GeometryCollectionComponent);
 
-					GeometryCollectionComponent->MarkRenderStateDirty();
-				}
+				GeometryCollectionComponent->MarkRenderStateDirty();
 			}
 		}
-
-		GCurrentLevelEditingViewportClient->Invalidate();
 	}
+
+	GCurrentLevelEditingViewportClient->Invalidate();
 }
 
 
@@ -645,7 +696,16 @@ int32 FFractureEditorModeToolkit::GetLevelCount()
 
 void FFractureEditorModeToolkit::OnSetLevelViewValue(int32 NewValue)
 {
-	FractureLevel = NewValue;//  * (float) GetLevelCount();
+	FScopedTransaction Transaction(LOCTEXT("SetLevelViewValue", "Adjust View Level"));
+
+	UFractureSettings* FractureSettings = GetMutableDefault<UFractureSettings>();
+	FractureSettings->FractureLevel = NewValue;
+	OnLevelViewValueChanged();
+}
+
+void FFractureEditorModeToolkit::OnLevelViewValueChanged()
+{
+	int32 FractureLevel = GetLevelViewValue();
 
 	USelection* SelectionSet = GEditor->GetSelectedActors();
 
@@ -676,19 +736,21 @@ void FFractureEditorModeToolkit::OnSetLevelViewValue(int32 NewValue)
 void FFractureEditorModeToolkit::ViewUpOneLevel()
 {
 	int32 CountMax = GetLevelCount() + 1;
-	int32 NewLevel = ((FractureLevel + CountMax) % CountMax) - 1;
+	int32 NewLevel = ((GetLevelViewValue() + CountMax) % CountMax) - 1;
 	OnSetLevelViewValue(NewLevel);
 }
 
 void FFractureEditorModeToolkit::ViewDownOneLevel()
 {
 	int32 CountMax = GetLevelCount() + 1;
-	int32 NewLevel = ((FractureLevel + CountMax + 2 ) % CountMax) - 1;
+	int32 NewLevel = ((GetLevelViewValue() + CountMax + 2 ) % CountMax) - 1;
 	OnSetLevelViewValue(NewLevel);
 }
 
-TSharedRef<SWidget> FFractureEditorModeToolkit::GetLevelViewMenuContent()
+TSharedRef<SWidget> FFractureEditorModeToolkit::GetLevelViewMenuContent(TSharedRef<IPropertyHandle> PropertyHandle)
 {
+	int32 FractureLevel = GetLevelViewValue();
+
 	FMenuBuilder MenuBuilder(true, GetToolkitCommands());
 
 	MenuBuilder.AddMenuEntry(
@@ -696,7 +758,7 @@ TSharedRef<SWidget> FFractureEditorModeToolkit::GetLevelViewMenuContent()
 		LOCTEXT("LevelMenuAllTooltip", "View All Leaf Bones in this Geometry Collection"),
 		FSlateIcon(),
 		FUIAction(
-			FExecuteAction::CreateSP(this, &FFractureEditorModeToolkit::OnSetLevelViewValue, -1),
+			FExecuteAction::CreateLambda([=] { PropertyHandle->SetValue(-1); } ),
 			FCanExecuteAction(),
 			FGetActionCheckState::CreateLambda([=] {return FractureLevel == -1 ? ECheckBoxState::Checked : ECheckBoxState::Unchecked;})
 		)
@@ -711,9 +773,9 @@ TSharedRef<SWidget> FFractureEditorModeToolkit::GetLevelViewMenuContent()
 			FText::Format(LOCTEXT("LevelMenuNTooltip", "View Level {0} in this Geometry Collecdtion"), FText::AsNumber(i)),
 			FSlateIcon(),
 			FUIAction(
-				FExecuteAction::CreateSP(this, &FFractureEditorModeToolkit::OnSetLevelViewValue, i),
+				FExecuteAction::CreateLambda([=] { PropertyHandle->SetValue(i); } ),
 				FCanExecuteAction(),
-				FGetActionCheckState::CreateLambda([=] {return FractureLevel == -1 ? ECheckBoxState::Checked : ECheckBoxState::Unchecked;})
+				FGetActionCheckState::CreateLambda([=] {return FractureLevel == i ? ECheckBoxState::Checked : ECheckBoxState::Unchecked;})
 			)
 		);
 	}
@@ -739,9 +801,11 @@ void FFractureEditorModeToolkit::SetActiveTool(UFractureTool* InActiveTool)
 	CommonSettings->OwnerTool = ActiveTool;
 
 	TArray<UObject*> Settings;
+	Settings.Add(GetMutableDefault<UFractureSettings>());
+
 	if (ActiveTool != nullptr)
 	{
-		Settings = ActiveTool->GetSettingsObjects();
+		Settings.Append(ActiveTool->GetSettingsObjects());
 
 		ActiveTool->FractureContextChanged();
 	}
@@ -1180,6 +1244,7 @@ void FFractureEditorModeToolkit::OnCluster()
 
 void FFractureEditorModeToolkit::OnUncluster()
 {
+	int32 FractureLevel = GetLevelViewValue();
 	TSet<UGeometryCollectionComponent*> GeomCompSelection;
 	GetSelectedGeometryCollectionComponents(GeomCompSelection);
 	for (UGeometryCollectionComponent* GeometryCollectionComponent : GeomCompSelection)
@@ -1318,7 +1383,6 @@ void FFractureEditorModeToolkit::GenerateAsset()
 	OpenGenerateAssetDialog(SelectedActors);
 }
 
-
 void FFractureEditorModeToolkit::OpenGenerateAssetDialog(TArray<AActor*>& Actors)
 {
 	TSharedPtr<SWindow> PickAssetPathWindow;
@@ -1595,6 +1659,8 @@ void FFractureEditorModeToolkit::UpdateExplodedVectors(UGeometryCollectionCompon
 	TSharedPtr<FGeometryCollection, ESPMode::ThreadSafe> GeometryCollectionPtr = GeometryCollectionComponent->GetRestCollection()->GetGeometryCollection();
 	const FGeometryCollection* OutGeometryCollectionConst = GeometryCollectionPtr.Get();
 
+	float ExplodeAmount = GetExplodedViewValue();
+
 	if (FMath::IsNearlyEqual(ExplodeAmount, 0.0f))
 	{
 		if (OutGeometryCollectionConst->HasAttribute("ExplodedVector", FGeometryCollection::TransformGroup))
diff --git a/Engine/Plugins/Experimental/ChaosEditor/Source/FractureEditor/Private/FractureSettings.cpp b/Engine/Plugins/Experimental/ChaosEditor/Source/FractureEditor/Private/FractureSettings.cpp
new file mode 100644
--- /dev/null
+++ b/Engine/Plugins/Experimental/ChaosEditor/Source/FractureEditor/Private/FractureSettings.cpp
@@ -0,0 +1,15 @@
+// Copyright Epic Games, Inc. All Rights Reserved.
+
+#include "FractureSettings.h"
+
+#if WITH_EDITOR
+
+UFractureSettings::UFractureSettings(const FObjectInitializer& ObjInit)
+	: Super(ObjInit)
+	, ExplodeAmount(0.0)
+	, FractureLevel(-1)
+{}
+
+#endif
+
+
diff --git a/Engine/Plugins/Experimental/ChaosEditor/Source/FractureEditor/Public/FractureEditorModeToolkit.h b/Engine/Plugins/Experimental/ChaosEditor/Source/FractureEditor/Public/FractureEditorModeToolkit.h
--- a/Engine/Plugins/Experimental/ChaosEditor/Source/FractureEditor/Public/FractureEditorModeToolkit.h
+++ b/Engine/Plugins/Experimental/ChaosEditor/Source/FractureEditor/Public/FractureEditorModeToolkit.h
@@ -7,8 +7,10 @@
 #include "Toolkits/BaseToolkit.h"
 #include "Widgets/Input/SCheckBox.h"
 #include "Widgets/Layout/SSplitter.h"
+#include "IDetailCustomization.h"
 
 class IDetailsView;
+class IPropertyHandle;
 class SScrollBox;
 class UFractureTool;
 class UEditableMesh;
@@ -18,17 +20,31 @@ class SGeometryCollectionOutliner;
 class AGeometryCollectionActor;
 class UGeometryCollectionComponent;
 class UGeometryCollection;
+class FFractureEditorModeToolkit;
 
 namespace GeometryCollection
 {
 enum class ESelectionMode: uint8;
 }
 
+class FFractureViewSettingsCustomization : public IDetailCustomization
+{
+public:
+	FFractureViewSettingsCustomization(FFractureEditorModeToolkit* FractureToolkit);
+	static TSharedRef<IDetailCustomization> MakeInstance(FFractureEditorModeToolkit* FractureToolkit);
+
+	virtual void CustomizeDetails(IDetailLayoutBuilder& DetailBuilder) override;
+
+	private:
+		FFractureEditorModeToolkit* Toolkit;
+};
+
 class FFractureEditorModeToolkit : public FModeToolkit, public FGCObject
 {
 public:
 
 	FFractureEditorModeToolkit();
+	~FFractureEditorModeToolkit();
 	
 	/** FModeToolkit interface */
 	virtual void Init(const TSharedPtr<IToolkitHost>& InitToolkitHost) override;
@@ -58,9 +74,16 @@ public:
 	bool GetShowBoneColors() const;
 	void OnSetExplodedViewValue(float NewValue);
 	void OnSetLevelViewValue(int32 NewValue);
-	void OnSetShowBoneColors();
+	void OnToggleShowBoneColors();
+	void OnSetShowBoneColors(bool NewValue);
 
-	TSharedRef<SWidget> GetLevelViewMenuContent();
+	void OnExplodedViewValueChanged();
+	void OnLevelViewValueChanged();
+
+	// Update any View Property Changes 
+	void OnObjectPostEditChange( UObject* Object, FPropertyChangedEvent& PropertyChangedEvent );
+
+	TSharedRef<SWidget> GetLevelViewMenuContent(TSharedRef<IPropertyHandle> PropertyHandle);
 	TSharedRef<SWidget> GetViewMenuContent();
 
 
@@ -98,7 +121,6 @@ public:
 
 	static void AddSingleRootNodeIfRequired(UGeometryCollection* GeometryCollectionObject);
 	static void AddAdditionalAttributesIfRequired(UGeometryCollection* GeometryCollectionObject);
-	// static int32 GetLevelCount();
 	int32 GetLevelCount();
 
 	FText GetStatisticsSummary() const;
@@ -123,16 +145,16 @@ protected:
 	static bool IsStaticMeshSelected();
 	static bool IsSelectedActorsInEditorWorld();
 	void UpdateExplodedVectors(UGeometryCollectionComponent* GeometryCollectionComponent) const;
+
 private:
 	void OnOutlinerBoneSelectionChanged(UGeometryCollectionComponent* RootComponent, TArray<int32>& SelectedBones);
 	void BindCommands();
-private:
-	float ExplodeAmount;
-	int32 FractureLevel;
 
+private:
 	UFractureTool* ActiveTool;
 
 	TSharedPtr<IDetailsView> DetailsView;
+	TSharedPtr<IDetailsView> ViewSettingsDetailsView;
 	TSharedPtr<SWidget> ToolkitWidget;
 	TSharedPtr<SGeometryCollectionOutliner> OutlinerView;
 
diff --git a/Engine/Plugins/Experimental/ChaosEditor/Source/FractureEditor/Public/FractureSettings.h b/Engine/Plugins/Experimental/ChaosEditor/Source/FractureEditor/Public/FractureSettings.h
new file mode 100644
--- /dev/null
+++ b/Engine/Plugins/Experimental/ChaosEditor/Source/FractureEditor/Public/FractureSettings.h
@@ -0,0 +1,26 @@
+// Copyright Epic Games, Inc. All Rights Reserved.
+
+#pragma once
+
+#include "UObject/Object.h"
+
+#include "FractureSettings.generated.h"
+
+/** Settings specifically related to viewing fractured meshes **/
+UCLASS()
+class UFractureSettings: public UObject
+{
+
+	GENERATED_BODY()
+public:
+	UFractureSettings(const FObjectInitializer& ObjInit);
+
+	/** How much to seperate the drawing of the bones  */
+	UPROPERTY(EditAnywhere, Category = ViewSettings, meta = (DisplayName = "Explode Amount", UIMin = "0.0", UIMax = "1.0", ClampMin = "0.0", ClampMax = "1.0"))
+	float ExplodeAmount;
+
+	/** Current level of the geometry collection displayed  */
+	UPROPERTY(EditAnywhere, Category = ViewSettings, meta = (DisplayName = "Fracture Level", UIMin = "-1" ))
+	int32 FractureLevel;
+
+};
\ No newline at end of file
diff --git a/Engine/Plugins/MeshPainting/Source/MeshPaintEditorMode/Private/MeshPaintMode.cpp b/Engine/Plugins/MeshPainting/Source/MeshPaintEditorMode/Private/MeshPaintMode.cpp
--- a/Engine/Plugins/MeshPainting/Source/MeshPaintEditorMode/Private/MeshPaintMode.cpp
+++ b/Engine/Plugins/MeshPainting/Source/MeshPaintEditorMode/Private/MeshPaintMode.cpp
@@ -118,7 +118,6 @@ UMeshPaintMode::UMeshPaintMode()
 {
 	SettingsClass = UMeshPaintModeSettings::StaticClass();
 	ToolsContextClass = UMeshToolsContext::StaticClass();
-	CurrentPaletteName = MeshPaintMode_Color;
 
 	FModuleManager::Get().LoadModule("EditorStyle");
 
@@ -145,14 +144,20 @@ void UMeshPaintMode::Enter()
 	RegisterTool(ToolManagerCommands.TexturePaint, TexturePaintToolName, NewObject<UMeshTexturePaintingToolBuilder>());
 	UpdateSelectedMeshes();
 
+
+	// Toolkit
+	PaletteChangedHandle = Toolkit->OnPaletteChanged().AddUObject(this, &UMeshPaintMode::UpdateOnPaletteChange);
+
 	// disable tool change tracking to activate default tool, and then switch to full undo/redo tracking mode
 	GetToolManager()->ConfigureChangeTrackingMode(EToolChangeTrackingMode::NoChangeTracking);
-	ActivateDefaultTool();
+	Toolkit->SetCurrentPalette(MeshPaintMode_Color);
 	GetToolManager()->ConfigureChangeTrackingMode(EToolChangeTrackingMode::FullUndoRedo);
 }
 
 void UMeshPaintMode::Exit()
 {
+	Toolkit->OnPaletteChanged().Remove(PaletteChangedHandle);
+
 	GEditor->OnEditorClose().RemoveAll(this);
 	OnResetViewMode();
 	const FMeshPaintEditorModeCommands& Commands = FMeshPaintEditorModeCommands::Get();
@@ -211,7 +216,7 @@ void UMeshPaintMode::Tick(FEditorViewportClient* ViewportClient, float DeltaTime
 	// Make sure that correct tab is visible for the current tool
 	// Note that currently Color and Weight mode share the same Select tool
 	FString ActiveTool = GetToolManager()->GetActiveToolName(EToolSide::Mouse);
-	FName ActiveTab = GetCurrentPaletteName();
+	FName ActiveTab = Toolkit->GetCurrentPalette();
 	FName TargetTab = ActiveTab;
 	if (ActiveTool == TexturePaintToolName || ActiveTool == TextureSelectToolName)
 	{
@@ -227,7 +232,7 @@ void UMeshPaintMode::Tick(FEditorViewportClient* ViewportClient, float DeltaTime
 	}
 	if ( TargetTab != ActiveTab)
 	{
-		GetModeManager()->InvokeToolPaletteTab(GetID(), TargetTab);
+		Toolkit->SetCurrentPalette(TargetTab);
 	}
 }
 
@@ -800,21 +805,22 @@ int32 UMeshPaintMode::GetNumberOfPendingPaintChanges()
 
 void UMeshPaintMode::ActivateDefaultTool()
 {
-	if (GetCurrentPaletteName() == UMeshPaintMode::MeshPaintMode_Color)
+	FName PaletteName = Toolkit->GetCurrentPalette();
+	if (PaletteName == UMeshPaintMode::MeshPaintMode_Color)
 	{
 		ToolsContext->StartTool(VertexSelectToolName);
 	}
-	if (GetCurrentPaletteName() == UMeshPaintMode::MeshPaintMode_Weights)
+	if (PaletteName == UMeshPaintMode::MeshPaintMode_Weights)
 	{
 		ToolsContext->StartTool(VertexSelectToolName);
 	}
-	if (GetCurrentPaletteName() == UMeshPaintMode::MeshPaintMode_Texture)
+	if (PaletteName == UMeshPaintMode::MeshPaintMode_Texture)
 	{
 		ToolsContext->StartTool(TextureSelectToolName);
 	}
 }
 
-void UMeshPaintMode::UpdateOnPaletteChange()
+void UMeshPaintMode::UpdateOnPaletteChange(FName NewPaletteName)
 {
 	UpdateSelectedMeshes();
 
@@ -823,15 +829,15 @@ void UMeshPaintMode::UpdateOnPaletteChange()
 
 	// figure out which tool we would like to be in based on currently-active tool
 	FString SwitchToTool;
-	if (GetCurrentPaletteName() == UMeshPaintMode::MeshPaintMode_Color)
+	if (NewPaletteName == UMeshPaintMode::MeshPaintMode_Color)
 	{
 		SwitchToTool = (bInAnyPaintTool) ? ColorPaintToolName : VertexSelectToolName;
 	}
-	else if (GetCurrentPaletteName() == UMeshPaintMode::MeshPaintMode_Weights)
+	else if (NewPaletteName == UMeshPaintMode::MeshPaintMode_Weights)
 	{
 		SwitchToTool = (bInAnyPaintTool) ? WeightPaintToolName: VertexSelectToolName;
 	}
-	else if (GetCurrentPaletteName() == UMeshPaintMode::MeshPaintMode_Texture)
+	else if (NewPaletteName == UMeshPaintMode::MeshPaintMode_Texture)
 	{
 		SwitchToTool = (bInAnyPaintTool) ? TexturePaintToolName : TextureSelectToolName;
 	}
diff --git a/Engine/Plugins/MeshPainting/Source/MeshPaintEditorMode/Private/MeshPaintMode.h b/Engine/Plugins/MeshPainting/Source/MeshPaintEditorMode/Private/MeshPaintMode.h
--- a/Engine/Plugins/MeshPainting/Source/MeshPaintEditorMode/Private/MeshPaintMode.h
+++ b/Engine/Plugins/MeshPainting/Source/MeshPaintEditorMode/Private/MeshPaintMode.h
@@ -72,7 +72,7 @@ protected:
 	virtual void OnToolEnded(UInteractiveToolManager* Manager, UInteractiveTool* Tool) override;
 	virtual void ActorSelectionChangeNotify() override;
 	virtual void ActivateDefaultTool() override;
-	virtual void UpdateOnPaletteChange() override;
+	virtual void UpdateOnPaletteChange(FName NewPalette);
 	// end UEdMode Interface
 	void UpdateSelectedMeshes();
 	void OnResetViewMode();
@@ -108,5 +108,7 @@ protected:
 	uint32 CachedVertexDataSize;
 	bool bRecacheVertexDataSize;
 
+	FDelegateHandle PaletteChangedHandle;
+
 };
 
diff --git a/Engine/Plugins/MeshPainting/Source/MeshPaintEditorMode/Private/MeshPaintModeCommands.cpp b/Engine/Plugins/MeshPainting/Source/MeshPaintEditorMode/Private/MeshPaintModeCommands.cpp
--- a/Engine/Plugins/MeshPainting/Source/MeshPaintEditorMode/Private/MeshPaintModeCommands.cpp
+++ b/Engine/Plugins/MeshPainting/Source/MeshPaintEditorMode/Private/MeshPaintModeCommands.cpp
@@ -55,7 +55,7 @@ void FMeshPaintEditorModeCommands::RegisterCommands()
 	ColorCommands.Add(ColorPaint);
 	UI_COMMAND(WeightPaint, "Paint", "Paint the mesh", EUserInterfaceActionType::ToggleButton, FInputChord());
 	WeightCommands.Add(WeightPaint);
-	UI_COMMAND(SwitchForeAndBackgroundColor, "Swap Colors", "Switches the Foreground and Background Colors used for Vertex Painting", EUserInterfaceActionType::Button, FInputChord(EKeys::X));
+	UI_COMMAND(SwitchForeAndBackgroundColor, "Swap", "Switches the Foreground and Background Colors used for Vertex Painting", EUserInterfaceActionType::Button, FInputChord(EKeys::X));
 	VertexCommands.Add(SwitchForeAndBackgroundColor);
 	UI_COMMAND(Fill, "Fill", "Fills the selected Meshes with the Paint Color", EUserInterfaceActionType::Button, FInputChord());
 	VertexCommands.Add(Fill);
@@ -75,10 +75,9 @@ void FMeshPaintEditorModeCommands::RegisterCommands()
 	VertexCommands.Add(Fix);
 	UI_COMMAND(PropagateVertexColorsToLODs, "All LODs", "Applied the Vertex Colors from LOD0 to all LOD levels", EUserInterfaceActionType::Button, FInputChord());
 	VertexCommands.Add(PropagateVertexColorsToLODs);
+
 	UI_COMMAND(CycleToPreviousLOD, "Previous LOD", "Cycles to the previous possible Mesh LOD to Paint on", EUserInterfaceActionType::Button, FInputChord(EKeys::B));
-	VertexCommands.Add(CycleToPreviousLOD);
 	UI_COMMAND(CycleToNextLOD, "Next LOD", "Cycles to the next possible Mesh LOD to Paint on", EUserInterfaceActionType::Button, FInputChord(EKeys::N));
-	VertexCommands.Add(CycleToNextLOD);
 
 	ColorCommands.Append(VertexCommands);
 	WeightCommands.Append(VertexCommands);
@@ -89,10 +88,10 @@ void FMeshPaintEditorModeCommands::RegisterCommands()
 	TextureCommands.Add(TextureSelect);
 	UI_COMMAND(TexturePaint, "Paint", "Paint the mesh", EUserInterfaceActionType::ToggleButton, FInputChord());
 	TextureCommands.Add(TexturePaint);
+	
 	UI_COMMAND(PreviousTexture, "Previous Texture", "Cycle To Previous Texture", EUserInterfaceActionType::Button, FInputChord(EKeys::Comma));
-	TextureCommands.Add(PreviousTexture);
 	UI_COMMAND(NextTexture, "Next Texture", "Cycle To Next Texture", EUserInterfaceActionType::Button, FInputChord(EKeys::Period));
-	TextureCommands.Add(NextTexture);
+
 	UI_COMMAND(PropagateTexturePaint, "Commit", "Commits Texture Painting Changes", EUserInterfaceActionType::Button, FInputChord(EKeys::C, EModifierKey::Control | EModifierKey::Shift));
 	TextureCommands.Add(PropagateTexturePaint);
 	UI_COMMAND(SaveTexturePaint, "Save", "Saves the Modified Textures for the selected Mesh Components", EUserInterfaceActionType::Button, FInputChord());
diff --git a/Engine/Source/Editor/EditorStyle/Private/StarshipStyle.cpp b/Engine/Source/Editor/EditorStyle/Private/StarshipStyle.cpp
--- a/Engine/Source/Editor/EditorStyle/Private/StarshipStyle.cpp
+++ b/Engine/Source/Editor/EditorStyle/Private/StarshipStyle.cpp
@@ -1773,35 +1773,67 @@ void FStarshipEditorStyle::FStyle::SetupGeneralStyles()
 	{
 
 		FToolBarStyle PaletteToolBarStyle = FStarshipCoreStyle::GetCoreStyle().GetWidgetStyle<FToolBarStyle>("SlimToolBar");
+
+		PaletteToolBarStyle.SetLabelStyle( GetParentStyle()->GetWidgetStyle<FTextBlockStyle>("SmallText") );
 		
-		PaletteToolBarStyle.SetBackground(FSlateColorBrush(FStyleColors::Background));
+		PaletteToolBarStyle.SetBackground(FSlateColorBrush(FStyleColors::Input));
+
 		PaletteToolBarStyle.SetLabelPadding(FMargin(0.0f, 4.0f, 0.0f, 0.0f));
-		PaletteToolBarStyle.SetCheckBoxPadding(FMargin(4.0f, 2.0f, 4.0f, 2.0f));
-		PaletteToolBarStyle.SetComboButtonPadding(FMargin(2.0f, 0.0f));
-		PaletteToolBarStyle.SetIndentedBlockPadding(FMargin(4.0f, 2.0f, 4.0f, 2.0f) );
-		PaletteToolBarStyle.SetBlockPadding(FMargin(4.0f, 2.0f, 4.0f, 2.0f));
 
-		Set("PaletteToolBar.MaxUniformToolbarSize", 50.f);
-		Set("PaletteToolbar.MinUniformToolbarSize", 22.f);
+		PaletteToolBarStyle.SetButtonPadding(       FMargin(2.0f, 2.0f));
+		PaletteToolBarStyle.SetCheckBoxPadding(     FMargin(1.0f, 1.0f));
+		PaletteToolBarStyle.SetComboButtonPadding(  FMargin(1.0f, 1.0f));
+		PaletteToolBarStyle.SetIndentedBlockPadding(FMargin(1.0f, 1.0f));
+		PaletteToolBarStyle.SetBlockPadding(        FMargin(1.0f, 1.0f));
+
+		Set( "PaletteToolBar.Tab",  FCheckBoxStyle()
+			.SetCheckBoxType(            ESlateCheckBoxType::ToggleButton)
+			.SetCheckedImage(            FSlateRoundedBoxBrush(FStyleColors::Input, 2.0))
+			.SetCheckedHoveredImage(     FSlateRoundedBoxBrush(FStyleColors::Input, 2.0))
+			.SetCheckedPressedImage(     FSlateRoundedBoxBrush(FStyleColors::Input, 2.0))
+			.SetUncheckedImage(          FSlateRoundedBoxBrush(FStyleColors::Dropdown, 2.0))
+			.SetUncheckedHoveredImage(   FSlateRoundedBoxBrush(FStyleColors::Dropdown, 2.0))
+			.SetUncheckedPressedImage(   FSlateRoundedBoxBrush(FStyleColors::Dropdown, 2.0))
+			.SetForegroundColor(         FStyleColors::Foreground)
+			.SetHoveredForegroundColor(  FStyleColors::ForegroundHover)
+			.SetPressedForegroundColor(  FStyleColors::ForegroundHover)
+			.SetCheckedForegroundColor(  FStyleColors::Primary)
+			.SetCheckedHoveredForegroundColor(FStyleColors::PrimaryHover)
+			.SetPadding(FMargin(2.f, 6.f))
+		);
 
+		Set("PaletteToolBar.MaxUniformToolbarSize", 64.f);
+		Set("PaletteToolBar.MinUniformToolbarSize", 64.f);
+		
 		FLinearColor ButtonHoverColor(FColor(81, 81, 81));
 
 			/* Create style for "ToolBar.ToggleButton" widget ... */
 		const FCheckBoxStyle ToolBarToggleButtonCheckBoxStyle = FCheckBoxStyle()
 				.SetCheckBoxType( ESlateCheckBoxType::ToggleButton )
-				.SetUncheckedImage( FSlateNoResource() )
-				.SetUncheckedPressedImage( 	BOX_BRUSH( "Common/RoundedSelection_16x", 4.0f/16.0f, SelectionColor_Pressed) )
-				.SetUncheckedHoveredImage( 	BOX_BRUSH( "Common/RoundedSelection_16x", 4.0f/16.0f, ButtonHoverColor ) )
-				.SetCheckedImage( 		   	BOX_BRUSH( "Common/RoundedSelection_16x", 4.0f/16.0f, SelectionColor_Pressed ) )
-				.SetCheckedHoveredImage( 	BOX_BRUSH( "Common/RoundedSelection_16x", 4.0f/16.0f, SelectionColor_Pressed ) )
-				.SetCheckedPressedImage( 	BOX_BRUSH( "Common/RoundedSelection_16x", 4.0f/16.0f, SelectionColor_Pressed ) );
+				.SetUncheckedImage(       FSlateRoundedBoxBrush(FStyleColors::Background, 2.0))
+				.SetUncheckedPressedImage(FSlateRoundedBoxBrush(FStyleColors::Background, 2.0))
+				.SetUncheckedHoveredImage(FSlateRoundedBoxBrush(FStyleColors::Background, 2.0))
+				.SetCheckedImage(         FSlateRoundedBoxBrush(FStyleColors::Input, 2.0))
+				.SetCheckedHoveredImage(  FSlateRoundedBoxBrush(FStyleColors::Input, 2.0))
+				.SetCheckedPressedImage(  FSlateRoundedBoxBrush(FStyleColors::Input, 2.0))
+				.SetForegroundColor(       FStyleColors::Foreground)
+				.SetHoveredForegroundColor(FStyleColors::ForegroundHover)
+				.SetPressedForegroundColor(FStyleColors::ForegroundHover)
+				.SetCheckedForegroundColor(FStyleColors::Primary)
+				.SetCheckedHoveredForegroundColor(FStyleColors::PrimaryHover)
+				.SetPadding(FMargin(2.f, 6.f));
 		
 		PaletteToolBarStyle.SetToggleButtonStyle(ToolBarToggleButtonCheckBoxStyle);
 		PaletteToolBarStyle.SetButtonStyle(FButtonStyle(Button)
-			.SetNormal(FSlateNoResource())
-			.SetNormalPadding(FMargin(2.0f, 0.0f, 2.0f, 0.0f))
-			.SetPressed(BOX_BRUSH("Common/RoundedSelection_16x", 4.0f / 16.0f, SelectionColor_Pressed))
-			.SetHovered(BOX_BRUSH("Common/RoundedSelection_16x", 4.0f / 16.0f, ButtonHoverColor)));
+			.SetNormal( FSlateRoundedBoxBrush(FStyleColors::Background, 2.0))
+			.SetPressed(FSlateRoundedBoxBrush(FStyleColors::Background, 2.0))
+			.SetHovered(FSlateRoundedBoxBrush(FStyleColors::Background, 2.0))
+			.SetNormalForeground(         FStyleColors::Foreground)
+			.SetHoveredForeground(        FStyleColors::ForegroundHover)
+			.SetPressedForeground(        FStyleColors::ForegroundHover)
+			.SetNormalPadding(FMargin(2.0f, 6.0f))
+			.SetPressedPadding(FMargin(2.0f, 6.0f))
+		);
 
 		Set("PaletteToolBar", PaletteToolBarStyle);
 
@@ -4371,6 +4403,13 @@ void FStarshipEditorStyle::FStyle::SetupLevelEditorStyle()
 		);	
 		Set( "ToolPalette.DockingWell", new FSlateColorBrush(FLinearColor(FColor(34, 34, 34, 255))));
 
+		Set( "ToolPalette.DockingLabel", FTextBlockStyle(NormalText)
+			.SetFont( DEFAULT_FONT( "Regular", 9 ) ) 
+			.SetShadowOffset(FVector2D(1, 1))
+			.SetShadowColorAndOpacity(FLinearColor(0, 0, 0, 0.9f))
+		);
+
+
 		Set("LevelEditor.SelectMode", new IMAGE_BRUSH_SVG("Starship/MainToolBar/select", Icon20x20));
 
 		Set( "LevelEditor.MeshPaintMode", new IMAGE_BRUSH_SVG( "Starship/MainToolBar/paint", Icon20x20 ) );
diff --git a/Engine/Source/Editor/FoliageEdit/Private/SFoliageEdit.cpp b/Engine/Source/Editor/FoliageEdit/Private/SFoliageEdit.cpp
--- a/Engine/Source/Editor/FoliageEdit/Private/SFoliageEdit.cpp
+++ b/Engine/Source/Editor/FoliageEdit/Private/SFoliageEdit.cpp
@@ -70,18 +70,6 @@ void SFoliageEdit::Construct(const FArguments& InArgs)
 		+ SVerticalBox::Slot()
 		.AutoHeight()
 		[
-			SNew(SHorizontalBox)
-			.Visibility_Lambda( [] () -> EVisibility { return GetDefault<UEditorStyleSettings>()->bEnableLegacyEditorModeUI ? EVisibility::Visible : EVisibility::Collapsed; } )
-			+ SHorizontalBox::Slot()
-			.AutoWidth()
-			.Padding(1.f, 5.f, 0.f, 5.f)
-			[
-				BuildToolBar()
-			]
-
-			+ SHorizontalBox::Slot()
-			.Padding(0.f, 2.f, 2.f, 0.f)
-			[
 				SNew(SBorder)
 				.BorderImage(FEditorStyle::GetBrush("ToolPanel.DarkGroupBorder"))
 				.Padding(StandardPadding)
@@ -435,99 +423,8 @@ void SFoliageEdit::Construct(const FArguments& InArgs)
 						]
 					]
 
-					// Actions
-					+ SVerticalBox::Slot()
-					.Padding(StandardPadding)
-					.AutoHeight()
-					[
-						SNew(SHeader)
-						.Visibility(this, &SFoliageEdit::GetVisibility_Actions)
-						[
-							SNew(STextBlock)
-							.Text(LOCTEXT("ActionsHeader", "Actions"))
-							.Font(StandardFont)
-						]
-					]
-
-					+ SVerticalBox::Slot()
-					.Padding(StandardPadding)
-					.AutoHeight()
-					[
-						SNew(SWrapBox)
-						.UseAllottedSize(true)
-						.Visibility(this, &SFoliageEdit::GetVisibility_SelectionOptions)
-
-						// Select all instances
-						+ SWrapBox::Slot()
-						.Padding(FMargin(0.f, 0.f, 6.f, 3.f))
-						[
-							SNew(SBox)
-							.WidthOverride(100.f)
-							.HeightOverride(25.f)
-							[
-								SNew(SButton)
-								.HAlign(HAlign_Center)
-								.VAlign(VAlign_Center)
-								.OnClicked_Lambda([this] () -> FReply { OnSelectAllInstances(); return FReply::Handled(); } )
-								.Text(LOCTEXT("SelectAllInstances", "Select All"))
-								.ToolTipText(LOCTEXT("SelectAllInstances_Tooltip", "Selects all foliage instances"))
-							]
-						]
-
-						// Select all invalid instances
-						+ SWrapBox::Slot()
-						.Padding(FMargin(0.f, 0.f, 6.f, 3.f))
-						[
-							SNew(SBox)
-							.WidthOverride(100.f)
-							.HeightOverride(25.f)
-							[
-								SNew(SButton)
-								.HAlign(HAlign_Center)
-								.VAlign(VAlign_Center)
-								.OnClicked_Lambda([this] () -> FReply { OnSelectInvalidInstances(); return FReply::Handled(); } )
-								.Text(LOCTEXT("SelectInvalidInstances", "Select Invalid"))
-								.ToolTipText(LOCTEXT("SelectInvalidInstances_Tooltip", "Selects all foliage instances that are not placed in a valid location"))
-						
-							]
-						]
-
-						// Deselect all
-						+ SWrapBox::Slot()
-						.Padding(FMargin(0.f, 0.f, 6.f, 3.f))
-						[
-							SNew(SBox)
-							.WidthOverride(100.f)
-							.HeightOverride(25.f)
-							[
-								SNew(SButton)
-								.HAlign(HAlign_Center)
-								.VAlign(VAlign_Center)
-								.OnClicked_Lambda([this] () ->FReply { OnDeselectAllInstances(); return FReply::Handled(); } )
-								.Text(LOCTEXT("DeselectAllInstances", "Deselect All"))
-								.ToolTipText(LOCTEXT("DeselectAllInstances_Tooltip", "Deselects all foliage instances"))
-							]
-						]
-
-						// Move to Current Level
-						+ SWrapBox::Slot()
-						.Padding(FMargin(0.f, 0.f, 6.f, 3.f))
-						[
-							SNew(SBox)
-							.WidthOverride(150.f)
-							.HeightOverride(25.f)
-							[
-								SNew(SButton)
-								.HAlign(HAlign_Center)
-								.VAlign(VAlign_Center)
-								.OnClicked_Lambda([this] () -> FReply { OnMoveSelectedInstancesToCurrentLevel(); return FReply::Handled(); } )
-								.Text(LOCTEXT("MoveSelectedInstancesToCurrentLevel", "Move to Current Level"))
-								.ToolTipText(LOCTEXT("MoveSelectedInstancesToCurrentLevel_Tooltip", "Move selected foliage instances to current level"))
-							]
-						]
-					]
+				
 				]
-			]
 		]
 
 		// Foliage Palette
@@ -646,11 +543,6 @@ void SFoliageEdit::CustomizeToolBarPalette(FToolBarBuilder& ToolBarBuilder)
 		EUserInterfaceActionType::ToggleButton
 	);
 
-	// Single Instance Options
-	ToolBarBuilder.AddComboButton(
-		FUIAction(),
-		FOnGetContent::CreateSP(this, &SFoliageEdit::GetSingleInstantiationModeMenuContent) );
-
 	//  Fill
 	ToolBarBuilder.AddToolBarButton(
 		FUIAction(
@@ -703,84 +595,6 @@ void SFoliageEdit::CustomizeToolBarPalette(FToolBarBuilder& ToolBarBuilder)
 		FSlateIcon(FEditorStyle::GetStyleSetName(), "FoliageEditMode.MoveToCurrentLevel")
 	);
 
-	ToolBarBuilder.AddSeparator();
-
-	TSharedPtr<INumericTypeInterface<float>> NumericInterface = MakeShareable(new FVariablePrecisionNumericInterface());
-	// Brush Size
-	{
-		TSharedRef<SWidget> BrushSizeWidget = SNew(SSpinBox<float>)
-			.Style(&FEditorStyle::Get().GetWidgetStyle<FSpinBoxStyle>("LandscapeEditor.SpinBox"))
-			.PreventThrottling(true)
-			.MinValue(0.0f)
-			.MaxValue(65536.0f)
-			.MaxSliderValue(8192.0f)
-			.SliderExponent(3.0f)
-			.Font(FCoreStyle::GetDefaultFontStyle("Regular", 11))
-			.MinDesiredWidth(40.f)
-			.TypeInterface(NumericInterface)
-			.Justification(ETextJustify::Center)
-			.IsEnabled(this, &SFoliageEdit::IsEnabled_BrushSize) 
-			.Value_Lambda( [this] { return GetRadius().GetValue(); } )
-			.OnValueChanged(this, &SFoliageEdit::SetRadius);
-
-		ToolBarBuilder.AddToolBarWidget(BrushSizeWidget, LOCTEXT("BrushSize", "Size") );
-	}
-
-	// Paint Density
-	{
-		TSharedRef<SWidget> PaintDensityWidget = SNew(SSpinBox<float>)
-			.Style(&FEditorStyle::Get().GetWidgetStyle<FSpinBoxStyle>("LandscapeEditor.SpinBox"))
-			.PreventThrottling(true)
-			.MinValue(0.0f)
-			.MaxValue(1.0f)
-			.MaxSliderValue(1.0f)
-			.Font(FCoreStyle::GetDefaultFontStyle("Regular", 11))
-			.MinDesiredWidth(40.f)
-			.TypeInterface(NumericInterface)
-			.Justification(ETextJustify::Center)
-			.IsEnabled(this, &SFoliageEdit::IsEnabled_PaintDensity) 
-			.Value_Lambda( [this] { return GetPaintDensity().GetValue(); } )
-			.OnValueChanged(this, &SFoliageEdit::SetPaintDensity);
-
-		ToolBarBuilder.AddToolBarWidget(PaintDensityWidget, LOCTEXT("Density", "Density") );
-	}
-
-	// Erase Density
-	{
-		TSharedRef<SWidget> EraseDensityWidget = SNew(SSpinBox<float>)
-			.Style(&FEditorStyle::Get().GetWidgetStyle<FSpinBoxStyle>("LandscapeEditor.SpinBox"))
-			.PreventThrottling(true)
-			.MinValue(0.0f)
-			.MaxValue(1.0f)
-			.MaxSliderValue(1.0f)
-			.Font(FCoreStyle::GetDefaultFontStyle("Regular", 11))
-			.MinDesiredWidth(40.f)
-			.TypeInterface(NumericInterface)
-			.Justification(ETextJustify::Center)
-			.IsEnabled(this, &SFoliageEdit::IsEnabled_EraseDensity) 
-			.Value_Lambda( [this] { return GetEraseDensity().GetValue(); } )
-			.OnValueChanged(this, &SFoliageEdit::SetEraseDensity);
-
-		ToolBarBuilder.AddToolBarWidget(EraseDensityWidget, LOCTEXT("EraseDensity", "Er. Dens.") );
-	}
-
-	// Filter Foliage Placement 
-	ToolBarBuilder.AddComboButton(
-		FUIAction(),
-		FOnGetContent::CreateSP(this, &SFoliageEdit::MakeFilterMenu),
-		LOCTEXT("Filter", "Filter"),
-		LOCTEXT("FilterTooltip", "Filter where foliage instances can be placed."),
-		FSlateIcon(FEditorStyle::GetStyleSetName(), "FoliageEditMode.Filter")
-	);
-
-	// Currently the only real setting is "Place in current Level"
-	ToolBarBuilder.AddComboButton(
-		FUIAction(),
-		FOnGetContent::CreateSP(this, &SFoliageEdit::MakeSettingsMenu),
-		LOCTEXT("Settings", "Settings"),
-		FText(),
-		FSlateIcon(FEditorStyle::GetStyleSetName(), "FoliageEditMode.Settings")
-	);
 }
 
 TSharedRef<SWidget> SFoliageEdit::MakeFilterMenu()
@@ -922,7 +736,7 @@ FText SFoliageEdit::GetFoliageEditorErrorText() const
 
 TSharedRef<SWidget> SFoliageEdit::BuildToolBar()
 {
-	FToolBarBuilder Toolbar(FoliageEditMode->UICommandList, FMultiBoxCustomization::None, nullptr, Orient_Vertical);
+	FVerticalToolBarBuilder Toolbar(FoliageEditMode->UICommandList, FMultiBoxCustomization::None);
 	Toolbar.SetLabelVisibility(EVisibility::Collapsed);
 	Toolbar.SetStyle(&FEditorStyle::Get(), "FoliageEditToolbar");
 	{
diff --git a/Engine/Source/Editor/LandscapeEditor/Private/LandscapeEdMode.cpp b/Engine/Source/Editor/LandscapeEditor/Private/LandscapeEdMode.cpp
--- a/Engine/Source/Editor/LandscapeEditor/Private/LandscapeEdMode.cpp
+++ b/Engine/Source/Editor/LandscapeEditor/Private/LandscapeEdMode.cpp
@@ -556,11 +556,19 @@ void FEdModeLandscape::Enter()
 	UpdateTargetList();
 	UpdateShownLayerList();
 
+	// Create the landscape editor window
+	if (!Toolkit.IsValid())
+	{
+		Toolkit = MakeShareable(new FLandscapeToolKit);
+		Toolkit->Init(Owner->GetToolkitHost());
+	}
+
 	// Initialize current tool prior to creating the landscape toolkit in case it has a dependency on it
 	if (LandscapeList.Num() == 0)
 	{
 		SetCurrentToolMode("ToolMode_Manage", false);
 		SetCurrentTool("NewLandscape");
+		Toolkit->SetCurrentPalette("ToolMode_Manage");
 	}
 	else
 	{
@@ -568,20 +576,15 @@ void FEdModeLandscape::Enter()
 		{
 			SetCurrentToolMode("ToolMode_Sculpt", false);
 			SetCurrentTool("Sculpt");
+			Toolkit->SetCurrentPalette("ToolMode_Sculpt");
 		}
 		else
 		{
 			SetCurrentTool(CurrentToolMode->CurrentToolName);
+			Toolkit->SetCurrentPalette(CurrentToolMode->ToolModeName);
 		}
 	}
 
-	// Create the landscape editor window
-	if (!Toolkit.IsValid())
-	{
-		Toolkit = MakeShareable(new FLandscapeToolKit);
-		Toolkit->Init(Owner->GetToolkitHost());
-	}
-
 	// Force real-time viewports.  We'll back up the current viewport state so we can restore it when the
 	// user exits this mode.
 	const bool bWantRealTime = true;
@@ -1957,7 +1960,10 @@ void FEdModeLandscape::SetCurrentToolMode(FName ToolModeName, bool bRestoreCurre
 					}
 					SetCurrentTool(CurrentToolMode->CurrentToolName, CurrentToolMode->CurrentTargetLayerName);
 				}
-				GetModeManager()->InvokeToolPaletteTab(GetID(), ToolModeName);
+				if (Toolkit.IsValid())
+				{
+					Toolkit->SetCurrentPalette(ToolModeName);
+				}
 				break;
 			}
 		}
diff --git a/Engine/Source/Editor/LandscapeEditor/Private/LandscapeEditorDetails.cpp b/Engine/Source/Editor/LandscapeEditor/Private/LandscapeEditorDetails.cpp
--- a/Engine/Source/Editor/LandscapeEditor/Private/LandscapeEditorDetails.cpp
+++ b/Engine/Source/Editor/LandscapeEditor/Private/LandscapeEditorDetails.cpp
@@ -8,6 +8,7 @@
 #include "Widgets/Input/SComboButton.h"
 #include "Widgets/Input/SMultiLineEditableTextBox.h"
 #include "Widgets/Input/SSpinBox.h"
+#include "Widgets/Input/SSegmentedControl.h"
 #include "LandscapeEdMode.h"
 #include "LandscapeEditorDetailCustomization_NewLandscape.h"
 #include "LandscapeEditorDetailCustomization_ResizeLandscape.h"
@@ -83,46 +84,57 @@ void FLandscapeEditorDetails::CustomizeDetails(IDetailLayoutBuilder& DetailBuild
 		return;
 	}
 
-	FToolSelectorBuilder ToolBrushSelectorButtons(CommandList, FMultiBoxCustomization::None);
+	// Custom Brush Selectors 
+
+	TSharedPtr<SSegmentedControl<FName>> BrushSelectors = SNew(SSegmentedControl<FName>)
+	.OnValueChanged(this, &FLandscapeEditorDetails::SetBrushCommand)
+	.Value(this, &FLandscapeEditorDetails::GetCurrentBrushFName) ;
+	for (FName BrushName : LandscapeEdMode->CurrentTool->ValidBrushes)
 	{
-		FUIAction ToolSelectorUIAction;
-		//ToolSelectorUIAction.IsActionVisibleDelegate.BindSP(this, &FLandscapeEditorDetails::GetToolSelectorIsVisible);
-		ToolBrushSelectorButtons.AddComboButton(
-			ToolSelectorUIAction,
-			FOnGetContent::CreateSP(this, &FLandscapeEditorDetails::GetToolSelector),
-			LOCTEXT("ToolSelector", "Tool"),
-			TAttribute<FText>(this, &FLandscapeEditorDetails::GetCurrentToolName),
-			LOCTEXT("ToolSelector.Tooltip", "Select Tool"),
-			TAttribute<FSlateIcon>(this, &FLandscapeEditorDetails::GetCurrentToolIcon)
-			);
-
-		FUIAction BrushSelectorUIAction;
-		BrushSelectorUIAction.IsActionVisibleDelegate.BindSP(this, &FLandscapeEditorDetails::GetBrushSelectorIsVisible);
-		ToolBrushSelectorButtons.AddComboButton(
-			BrushSelectorUIAction,
-			FOnGetContent::CreateSP(this, &FLandscapeEditorDetails::GetBrushSelector),
-			LOCTEXT("BrushSelector", "Brush"),
-			TAttribute<FText>(this, &FLandscapeEditorDetails::GetCurrentBrushName),
-			LOCTEXT("BrushSelector.Tooltip", "Select Brush"),
-			TAttribute<FSlateIcon>(this, &FLandscapeEditorDetails::GetCurrentBrushIcon)
-			);
-
-		FUIAction BrushFalloffSelectorUIAction;
-		BrushFalloffSelectorUIAction.IsActionVisibleDelegate.BindSP(this, &FLandscapeEditorDetails::GetBrushFalloffSelectorIsVisible);
-		ToolBrushSelectorButtons.AddComboButton(
-			BrushFalloffSelectorUIAction,
-			FOnGetContent::CreateSP(this, &FLandscapeEditorDetails::GetBrushFalloffSelector),
-			LOCTEXT("BrushFalloffSelector", "Falloff"),
-			TAttribute<FText>(this, &FLandscapeEditorDetails::GetCurrentBrushFalloffName),
-			LOCTEXT("BrushFalloffSelector.Tooltip", "Select Brush Falloff Type"),
-			TAttribute<FSlateIcon>(this, &FLandscapeEditorDetails::GetCurrentBrushFalloffIcon)
-			);
+		TSharedPtr<FUICommandInfo> Command = FLandscapeEditorCommands::Get().NameToCommandMap.FindRef(BrushName);
+		if (Command.IsValid())
+		{
+			BrushSelectors->AddSlot(BrushName)
+			.Icon(Command->GetIcon().GetIcon());
+		}
 	}
 
+	TSharedPtr<SSegmentedControl<FName>> FalloffSelectors = SNew(SSegmentedControl<FName>)
+	.OnValueChanged(this, &FLandscapeEditorDetails::SetBrushCommand)
+	.Value(this, &FLandscapeEditorDetails::GetCurrentBrushFalloffFName)
+	+ SSegmentedControl<FName>::Slot(FName("Circle_Smooth")).Icon(FLandscapeEditorCommands::Get().CircleBrush_Smooth->GetIcon().GetIcon())
+	+ SSegmentedControl<FName>::Slot(FName("Circle_Linear")).Icon(FLandscapeEditorCommands::Get().CircleBrush_Linear->GetIcon().GetIcon())
+	+ SSegmentedControl<FName>::Slot(FName("Circle_Spherical")).Icon(FLandscapeEditorCommands::Get().CircleBrush_Spherical->GetIcon().GetIcon())
+	+ SSegmentedControl<FName>::Slot(FName("Circle_Tip")).Icon(FLandscapeEditorCommands::Get().CircleBrush_Tip->GetIcon().GetIcon());
+
 	LandscapeEditorCategory.AddCustomRow(FText::GetEmpty())
-	.Visibility(TAttribute<EVisibility>::Create(TAttribute<EVisibility>::FGetter::CreateSP(this, &FLandscapeEditorDetails::GetToolSelectorVisibility)))
+	.Visibility(TAttribute<EVisibility>::Create(TAttribute<EVisibility>::FGetter::CreateSP(this, &FLandscapeEditorDetails::GetBrushSelectorVisibility)))
+	.NameContent()
+	.HAlign(HAlign_Right)
 	[
-		ToolBrushSelectorButtons.MakeWidget()
+		SNew(STextBlock)
+		.TextStyle( &FCoreStyle::Get().GetWidgetStyle<FTextBlockStyle>( "SmallText" ) )
+		.Text(LOCTEXT("BrushSelector", "Brush Type"))
+
+	]
+	.ValueContent()
+	[
+		BrushSelectors.ToSharedRef()
+	];
+
+	LandscapeEditorCategory.AddCustomRow(FText::GetEmpty())
+	.Visibility(TAttribute<EVisibility>::Create(TAttribute<EVisibility>::FGetter::CreateSP(this, &FLandscapeEditorDetails::GetBrushFalloffSelectorVisibility)))
+	.NameContent()
+	.HAlign(HAlign_Right)
+	[
+		SNew(STextBlock)
+		.TextStyle( &FCoreStyle::Get().GetWidgetStyle<FTextBlockStyle>( "SmallText" ) )
+		.Text(LOCTEXT("BrushFalloff", "Brush Falloff"))
+
+	]
+	.ValueContent()
+	[
+		FalloffSelectors.ToSharedRef()
 	];
 
 	// Tools:
@@ -155,151 +167,6 @@ void FLandscapeEditorDetails::CustomizeDetails(IDetailLayoutBuilder& DetailBuild
 	Customization_TargetLayers->CustomizeDetails(DetailBuilder);
 }
 
-void FLandscapeEditorDetails::CustomizeToolBarPalette(FToolBarBuilder& ToolBarBuilder, const TSharedRef<FLandscapeToolKit> LandscapeToolkit)
-{
-
-	FEdModeLandscape* LandscapeEdMode = GetEditorMode();
-	CommandList = LandscapeEdMode->GetUICommandList();
-
-	TSharedPtr<INumericTypeInterface<float>> NumericInterface = MakeShareable(new FVariablePrecisionNumericInterface());
-
-	// Tool Strength
-	{
-		FProperty* ToolStrengthProperty = LandscapeEdMode->UISettings->GetClass()->FindPropertyByName(GET_MEMBER_NAME_CHECKED(ULandscapeEditorObject, ToolStrength));
-		const FString& UIMinString = ToolStrengthProperty->GetMetaData("UIMin");
-		const FString& UIMaxString = ToolStrengthProperty->GetMetaData("UIMax");
-		const FString& SliderExponentString = ToolStrengthProperty->GetMetaData("SliderExponent");
-		float UIMin = TNumericLimits<float>::Lowest();
-		float UIMax = TNumericLimits<float>::Max();
-		TTypeFromString<float>::FromString(UIMin, *UIMinString);
-		TTypeFromString<float>::FromString(UIMax, *UIMaxString);
-		float SliderExponent = 1.0f;
-		if (SliderExponentString.Len())
-		{
-			TTypeFromString<float>::FromString(SliderExponent, *SliderExponentString);
-		}
-
-		TSharedRef<SWidget> StrengthControl = SNew(SSpinBox<float>)
-			.Style(&FEditorStyle::Get().GetWidgetStyle<FSpinBoxStyle>("LandscapeEditor.SpinBox"))
-			.PreventThrottling(true)
-			.MinValue(UIMin)
-			.MaxValue(UIMax)
-			.SliderExponent(SliderExponent)
-			.Font(FCoreStyle::GetDefaultFontStyle("Regular", 11))
-			.MinDesiredWidth(40.f)
-			.TypeInterface(NumericInterface)
-			.Justification(ETextJustify::Center)
-			.ContentPadding( FMargin(0.f, 2.f, 0.f, 0.f) )
-
-			.Visibility_Lambda( [ToolStrengthProperty, LandscapeToolkit]() -> EVisibility { return LandscapeToolkit->GetIsPropertyVisibleFromProperty(*ToolStrengthProperty) ? EVisibility::Visible : EVisibility::Collapsed;} )
-			.IsEnabled(this, &FLandscapeEditorDetails::IsBrushSetEnabled) 
-			.Value_Lambda([LandscapeEdMode]() -> float { return LandscapeEdMode->UISettings->ToolStrength; })
-			.OnValueChanged_Lambda([LandscapeEdMode](float NewValue) { LandscapeEdMode->UISettings->ToolStrength = NewValue; });
-
-		ToolBarBuilder.AddToolBarWidget( StrengthControl, LOCTEXT("BrushStrength", "Strength") );
-	}
-
-	//  Brush
-	FUIAction BrushSelectorUIAction;
-	BrushSelectorUIAction.IsActionVisibleDelegate.BindSP(this, &FLandscapeEditorDetails::GetBrushSelectorIsVisible);
-	BrushSelectorUIAction.CanExecuteAction.BindSP(this, &FLandscapeEditorDetails::IsBrushSetEnabled);
-
-	ToolBarBuilder.AddComboButton(
-		BrushSelectorUIAction,
-		FOnGetContent::CreateSP(this, &FLandscapeEditorDetails::GetBrushSelector),
-		TAttribute<FText>(this, &FLandscapeEditorDetails::GetCurrentBrushName),
-		LOCTEXT("BrushSelector.Tooltip", "Select Brush"),
-		TAttribute<FSlateIcon>(this, &FLandscapeEditorDetails::GetCurrentBrushIcon)
-		);
-
-	//  Brush Size 
-	{
-
-		FProperty* BrushRadiusProperty = LandscapeEdMode->UISettings->GetClass()->FindPropertyByName(GET_MEMBER_NAME_CHECKED(ULandscapeEditorObject, BrushRadius));
-		const FString& UIMinString = BrushRadiusProperty->GetMetaData("UIMin");
-		const FString& UIMaxString = BrushRadiusProperty->GetMetaData("UIMax");
-		const FString& SliderExponentString = BrushRadiusProperty->GetMetaData("SliderExponent");
-		float UIMin = TNumericLimits<float>::Lowest();
-		float UIMax = TNumericLimits<float>::Max();
-		TTypeFromString<float>::FromString(UIMin, *UIMinString);
-		TTypeFromString<float>::FromString(UIMax, *UIMaxString);
-		float SliderExponent = 1.0f;
-		if (SliderExponentString.Len())
-		{
-			TTypeFromString<float>::FromString(SliderExponent, *SliderExponentString);
-		}
-
-		TSharedRef<SWidget> SizeControl = 
-			SNew(SSpinBox<float>)
-			.Style(&FEditorStyle::Get().GetWidgetStyle<FSpinBoxStyle>("LandscapeEditor.SpinBox"))
-			.Visibility_Lambda( [BrushRadiusProperty, LandscapeToolkit]() -> EVisibility { return LandscapeToolkit->GetIsPropertyVisibleFromProperty(*BrushRadiusProperty) ? EVisibility::Visible : EVisibility::Collapsed;} )
-			.IsEnabled(this, &FLandscapeEditorDetails::IsBrushSetEnabled) 
-			.PreventThrottling(true)
-			.Value_Lambda([LandscapeEdMode]() -> float { return LandscapeEdMode->UISettings->BrushRadius; })
-			.OnValueChanged_Lambda([LandscapeEdMode](float NewValue) { LandscapeEdMode->UISettings->BrushRadius = NewValue; })
-
-			.MinValue(UIMin)
-			.MaxValue(UIMax)
-			.SliderExponent(SliderExponent)
-			.Font(FCoreStyle::GetDefaultFontStyle("Regular", 11))
-			.MinDesiredWidth(40.f)
-			.TypeInterface(NumericInterface)
-			.Justification(ETextJustify::Center);
-		ToolBarBuilder.AddToolBarWidget( SizeControl, LOCTEXT("BrushRadius", "Radius") );
-	}
-
-	//  Brush Falloff Curve
-	FUIAction BrushFalloffSelectorUIAction;
-	BrushFalloffSelectorUIAction.IsActionVisibleDelegate.BindSP(this, &FLandscapeEditorDetails::GetBrushFalloffSelectorIsVisible);
-	BrushFalloffSelectorUIAction.CanExecuteAction.BindSP(this, &FLandscapeEditorDetails::IsBrushSetEnabled);
-	ToolBarBuilder.AddComboButton(
-		BrushFalloffSelectorUIAction,
-		FOnGetContent::CreateSP(this, &FLandscapeEditorDetails::GetBrushFalloffSelector),
-		TAttribute<FText>(this, &FLandscapeEditorDetails::GetCurrentBrushFalloffName),
-		LOCTEXT("BrushFalloffSelector.Tooltip", "Select Brush Falloff Type"),
-		TAttribute<FSlateIcon>(this, &FLandscapeEditorDetails::GetCurrentBrushFalloffIcon)
-		);
-
-	//  Brush Falloff Percentage 
-	{
-		FProperty* BrushFalloffProperty = LandscapeEdMode->UISettings->GetClass()->FindPropertyByName(GET_MEMBER_NAME_CHECKED(ULandscapeEditorObject, BrushFalloff));
-
-		const FString& UIMinString = BrushFalloffProperty->GetMetaData("UIMin");
-		const FString& UIMaxString = BrushFalloffProperty->GetMetaData("UIMax");
-		const FString& SliderExponentString = BrushFalloffProperty->GetMetaData("SliderExponent");
-		float UIMin = TNumericLimits<float>::Lowest();
-		float UIMax = TNumericLimits<float>::Max();
-		TTypeFromString<float>::FromString(UIMin, *UIMinString);
-		TTypeFromString<float>::FromString(UIMax, *UIMaxString);
-		float SliderExponent = 1.0f;
-		if (SliderExponentString.Len())
-		{
-			TTypeFromString<float>::FromString(SliderExponent, *SliderExponentString);
-		}
-
-		TSharedRef<SWidget> FalloffControl = 
-			SNew(SSpinBox<float>)
-			.Style(&FEditorStyle::Get().GetWidgetStyle<FSpinBoxStyle>("LandscapeEditor.SpinBox"))
-			.Visibility_Lambda( [BrushFalloffProperty, LandscapeToolkit, this]() -> EVisibility { return GetBrushFalloffSelectorIsVisible() && LandscapeToolkit->GetIsPropertyVisibleFromProperty(*BrushFalloffProperty) ? EVisibility::Visible : EVisibility::Collapsed;} )
-			.IsEnabled(this, &FLandscapeEditorDetails::IsBrushSetEnabled) 
-			.PreventThrottling(true)
-			.Value_Lambda([LandscapeEdMode]() -> float { return LandscapeEdMode->UISettings->BrushFalloff; })
-			.OnValueChanged_Lambda([LandscapeEdMode](float NewValue) { LandscapeEdMode->UISettings->BrushFalloff = NewValue; })
-
-			.MinValue(UIMin)
-			.MaxValue(UIMax)
-			.SliderExponent(SliderExponent)
-			.Font(FCoreStyle::GetDefaultFontStyle("Regular", 11))
-			.MinDesiredWidth(40.f)
-			.TypeInterface(NumericInterface)
-			.Justification(ETextJustify::Center);
-
-		ToolBarBuilder.AddToolBarWidget( FalloffControl, LOCTEXT("BrushFalloff", "Falloff") );
-	}
-
-}
-
-
 FText FLandscapeEditorDetails::GetLocalizedName(FString Name)
 {
 	static bool bInitialized = false;
@@ -522,11 +389,6 @@ TSharedRef<SWidget> FLandscapeEditorDetails::GetToolSelector()
 
 bool FLandscapeEditorDetails::GetToolSelectorIsVisible() const
 {
-	if (!GetDefault<UEditorStyleSettings>()->bEnableLegacyEditorModeUI)
-	{
-		return false;
-	}
-
 	FEdModeLandscape* LandscapeEdMode = GetEditorMode();
 	if (LandscapeEdMode && LandscapeEdMode->CurrentTool)
 	{
@@ -553,6 +415,19 @@ EVisibility FLandscapeEditorDetails::GetToolSelectorVisibility() const
 	return EVisibility::Collapsed;
 }
 
+
+FName FLandscapeEditorDetails::GetCurrentBrushFName() const
+{
+	FEdModeLandscape* LandscapeEdMode = GetEditorMode();
+	if (LandscapeEdMode != NULL && LandscapeEdMode->CurrentBrush != NULL)
+	{
+		return LandscapeEdMode->LandscapeBrushSets[LandscapeEdMode->CurrentBrushSetIndex].BrushSetName;
+	}
+
+	return NAME_None;
+}
+
+
 FText FLandscapeEditorDetails::GetCurrentBrushName() const
 {
 	FEdModeLandscape* LandscapeEdMode = GetEditorMode();
@@ -637,6 +512,22 @@ bool FLandscapeEditorDetails::GetBrushSelectorIsVisible() const
 	return false;
 }
 
+EVisibility FLandscapeEditorDetails::GetBrushSelectorVisibility() const
+{
+	return GetBrushSelectorIsVisible() ? EVisibility::Visible : EVisibility::Collapsed;
+}
+
+FName FLandscapeEditorDetails::GetCurrentBrushFalloffFName() const
+{
+	FEdModeLandscape* LandscapeEdMode = GetEditorMode();
+	if (LandscapeEdMode != NULL && LandscapeEdMode->CurrentBrush != NULL && GetBrushFalloffSelectorIsVisible())
+	{
+		return LandscapeEdMode->CurrentBrush->GetBrushName();
+	}
+
+	return NAME_None;
+}
+
 FText FLandscapeEditorDetails::GetCurrentBrushFalloffName() const
 {
 	FEdModeLandscape* LandscapeEdMode = GetEditorMode();
@@ -665,6 +556,15 @@ FSlateIcon FLandscapeEditorDetails::GetCurrentBrushFalloffIcon() const
 	return FSlateIcon(FEditorStyle::GetStyleSetName(), "Default");
 }
 
+void FLandscapeEditorDetails::SetBrushCommand(FName InBrush)
+{
+	TSharedPtr<FUICommandInfo> BrushCommand = FLandscapeEditorCommands::Get().NameToCommandMap.FindRef(InBrush);
+	if (CommandList.IsValid() && BrushCommand.IsValid())
+	{
+		CommandList->ExecuteAction( BrushCommand.ToSharedRef() );
+	}
+}
+
 TSharedRef<SWidget> FLandscapeEditorDetails::GetBrushFalloffSelector()
 {
 	FEdModeLandscape* LandscapeEdMode = GetEditorMode();
@@ -702,6 +602,11 @@ bool FLandscapeEditorDetails::GetBrushFalloffSelectorIsVisible() const
 	return false;
 }
 
+EVisibility FLandscapeEditorDetails::GetBrushFalloffSelectorVisibility() const
+{
+	return GetBrushFalloffSelectorIsVisible() ? EVisibility::Visible : EVisibility::Collapsed;
+}
+
 bool FLandscapeEditorDetails::IsBrushSetEnabled() const
 {
 	FEdModeLandscape* LandscapeEdMode = GetEditorMode();
diff --git a/Engine/Source/Editor/LandscapeEditor/Private/LandscapeEditorDetails.h b/Engine/Source/Editor/LandscapeEditor/Private/LandscapeEditorDetails.h
--- a/Engine/Source/Editor/LandscapeEditor/Private/LandscapeEditorDetails.h
+++ b/Engine/Source/Editor/LandscapeEditor/Private/LandscapeEditorDetails.h
@@ -43,8 +43,6 @@ public:
 	/** IDetailCustomization interface */
 	virtual void CustomizeDetails(IDetailLayoutBuilder& DetailBuilder) override;
 
-	void CustomizeToolBarPalette(FToolBarBuilder& ToolBarBuilder, const TSharedRef<FLandscapeToolKit> LandscapeToolkit);
-
 protected:
 	static FText GetLocalizedName(FString Name);
 
@@ -59,15 +57,21 @@ protected:
 	bool GetToolSelectorIsVisible() const;
 	EVisibility GetToolSelectorVisibility() const;
 
+	FName GetCurrentBrushFName() const;
 	FText GetCurrentBrushName() const;
 	FSlateIcon GetCurrentBrushIcon() const;
 	TSharedRef<SWidget> GetBrushSelector();
 	bool GetBrushSelectorIsVisible() const;
+	EVisibility GetBrushSelectorVisibility() const;
 
+	FName GetCurrentBrushFalloffFName() const;
 	FText GetCurrentBrushFalloffName() const;
 	FSlateIcon GetCurrentBrushFalloffIcon() const;
 	TSharedRef<SWidget> GetBrushFalloffSelector();
 	bool GetBrushFalloffSelectorIsVisible() const;
+	EVisibility GetBrushFalloffSelectorVisibility() const;
+
+	void SetBrushCommand(FName);
 
 	bool IsBrushSetEnabled() const;
 
diff --git a/Engine/Source/Editor/LandscapeEditor/Private/SLandscapeEditor.cpp b/Engine/Source/Editor/LandscapeEditor/Private/SLandscapeEditor.cpp
--- a/Engine/Source/Editor/LandscapeEditor/Private/SLandscapeEditor.cpp
+++ b/Engine/Source/Editor/LandscapeEditor/Private/SLandscapeEditor.cpp
@@ -202,7 +202,6 @@ void FLandscapeToolKit::BuildToolPalette(FName PaletteName, class FToolBarBuilde
 		ToolBarBuilder.BeginSection("Manage");
 
 		ToolBarBuilder.AddToolBarButton(Commands.NewLandscape);
-		ToolBarBuilder.AddSeparator();	
 
 		ToolBarBuilder.AddToolBarButton(Commands.SelectComponentTool);
 		ToolBarBuilder.AddToolBarButton(Commands.AddComponentTool);
@@ -210,9 +209,7 @@ void FLandscapeToolKit::BuildToolPalette(FName PaletteName, class FToolBarBuilde
 		ToolBarBuilder.AddToolBarButton(Commands.MoveToLevelTool);
 		ToolBarBuilder.AddToolBarButton(Commands.ResizeLandscape);
 
-		ToolBarBuilder.AddSeparator();	
 		ToolBarBuilder.AddToolBarButton(Commands.SplineTool);
-		ToolBarBuilder.EndSection();
 	}
 
 	else if (PaletteName == LandscapeEditorNames::Sculpt)
@@ -238,18 +235,9 @@ void FLandscapeToolKit::BuildToolPalette(FName PaletteName, class FToolBarBuilde
 		}
 		ToolBarBuilder.AddToolBarButton(Commands.MirrorTool);
 
-		ToolBarBuilder.AddSeparator();
-
 		ToolBarBuilder.AddToolBarButton(Commands.RegionSelectTool);
 		ToolBarBuilder.AddToolBarButton(Commands.RegionCopyPasteTool);
 
-		ToolBarBuilder.AddSeparator();
-
-		if (BrushesWidgets)
-		{
-			BrushesWidgets->CustomizeToolBarPalette(ToolBarBuilder, SharedThis(this));
-		}
-
 	}
 
 	else if (PaletteName == LandscapeEditorNames::Paint)
@@ -263,13 +251,6 @@ void FLandscapeToolKit::BuildToolPalette(FName PaletteName, class FToolBarBuilde
 		{
 			ToolBarBuilder.AddToolBarButton(Commands.BlueprintBrushTool);
 		}
-
-		ToolBarBuilder.AddSeparator();	
-
-		if (BrushesWidgets)
-		{
-			BrushesWidgets->CustomizeToolBarPalette(ToolBarBuilder, SharedThis(this));
-		}
 	}
 
 }
diff --git a/Engine/Source/Editor/LevelEditor/Private/SLevelEditorToolBox.cpp b/Engine/Source/Editor/LevelEditor/Private/SLevelEditorToolBox.cpp
--- a/Engine/Source/Editor/LevelEditor/Private/SLevelEditorToolBox.cpp
+++ b/Engine/Source/Editor/LevelEditor/Private/SLevelEditorToolBox.cpp
@@ -16,7 +16,10 @@
 #include "Classes/EditorStyleSettings.h"
 #include "EdMode.h"
 #include "Widgets/Docking/SDockTab.h"
+#include "Widgets/Input/SCheckBox.h"
 #include "Widgets/Layout/SExpandableArea.h"
+#include "Widgets/Layout/SWidgetSwitcher.h"
+#include "Widgets/Layout/SUniformWrapPanel.h"
 #include "StatusBarSubsystem.h"
 
 #define LOCTEXT_NAMESPACE "SLevelEditorToolBox"
@@ -41,36 +44,40 @@ void SLevelEditorToolBox::Construct( const FArguments& InArgs, const TSharedRef<
 
 	ChildSlot
 	[
-		SNew( SVerticalBox )
-		+ SVerticalBox::Slot()
-		.AutoHeight()
-		.HAlign( HAlign_Left )
-		.Padding(1.f)
-		[
-			SAssignNew( ModeToolBarContainer, SBorder )
-			.BorderImage( FEditorStyle::GetBrush( "NoBorder" ) )
-			.Padding( FMargin(4, 0, 0, 0) )
-		]
-
-		+ SVerticalBox::Slot()
-		.FillHeight( 1.0f )
-		.Padding( 2, 0, 0, 0 )
+		SNew( SBorder )
+		.BorderImage( FEditorStyle::GetBrush( "ToolPanel.GroupBorder" ) )
+		.Padding(0.0f)
 		[
 			SNew( SVerticalBox )
-
-			+SVerticalBox::Slot()
+			+ SVerticalBox::Slot()
 			.AutoHeight()
+			.HAlign( HAlign_Left )
 			[
-				SAssignNew(ModeToolHeader, SBorder)
-				.BorderImage( FEditorStyle::GetBrush( "NoBorder" ) )
+				SAssignNew( ModeToolBarContainer, SBorder )
+				.BorderImage( FEditorStyle::GetBrush( "ToolPanel.GroupBorder" ) )
+				.Padding( FMargin(4, 0, 0, 0) )
 			]
 
 			+ SVerticalBox::Slot()
+			.FillHeight( 1.0f )
 			[
-				SAssignNew(InlineContentHolder, SBorder)
-				.BorderImage( FEditorStyle::GetBrush( "ToolPanel.GroupBorder" ) )
-				.Padding(0.f)
-				.Visibility( this, &SLevelEditorToolBox::GetInlineContentHolderVisibility )
+				SNew( SVerticalBox )
+
+				+SVerticalBox::Slot()
+				.Padding(0.0, 8.0, 0.0, 0.0)
+				.AutoHeight()
+				[
+					SAssignNew(ModeToolHeader, SBorder)
+					.BorderImage( FEditorStyle::GetBrush( "ToolPanel.GroupBorder" ) )
+				]
+
+				+ SVerticalBox::Slot()
+				.FillHeight(1)
+				[
+					SAssignNew(InlineContentHolder, SBorder)
+					.BorderImage( FEditorStyle::GetBrush( "ToolPanel.GroupBorder" ) )
+					.Visibility( this, &SLevelEditorToolBox::GetInlineContentHolderVisibility )
+				]
 			]
 		]
 	];
@@ -173,28 +180,81 @@ void SLevelEditorToolBox::UpdateInlineContent(const TSharedPtr<IToolkit>& Toolki
 
 			TSharedPtr<FModeToolkit> ModeToolkit = StaticCastSharedPtr<FModeToolkit>(Toolkit);
 
+			TSharedRef< SUniformWrapPanel> PaletteTabBox = SNew(SUniformWrapPanel)
+				.SlotPadding( FMargin(1.f, 2.f) )
+				.HAlign(HAlign_Center);
+
+			// Only show if there is more than one child in the switcher
+			PaletteTabBox->SetVisibility(TAttribute<EVisibility>::Create(TAttribute<EVisibility>::FGetter::CreateLambda([PaletteTabBox]() -> EVisibility { return PaletteTabBox->GetChildren()->Num() > 1 ? EVisibility::Visible : EVisibility::Collapsed; }) ));
+
+			// Also build the toolkit here 
+			TArray<FName> PaletteNames;
+			ModeToolkit->GetToolPaletteNames(PaletteNames);
+
+			TSharedPtr<FUICommandList> CommandList;
+			CommandList = ModeToolkit->GetToolkitCommands();
+
+			TSharedRef< SWidgetSwitcher > PaletteSwitcher = SNew(SWidgetSwitcher)
+			.WidgetIndex_Lambda( [PaletteNames, ModeToolkit] () -> int32 { 
+				int32 FoundIndex;
+				if (PaletteNames.Find(ModeToolkit->GetCurrentPalette(), FoundIndex))
+				{
+					return FoundIndex;	
+				}
+				return 0;
+			} );
+			
+			for(auto Palette : PaletteNames)
+			{
+				FName ToolbarCustomizationName = ModeToolkit->GetEditorMode() ? ModeToolkit->GetEditorMode()->GetModeInfo().ToolbarCustomizationName : ModeToolkit->GetScriptableEditorMode()->GetModeInfo().ToolbarCustomizationName;
+				FUniformToolBarBuilder ModeToolbarBuilder(CommandList, FMultiBoxCustomization(ToolbarCustomizationName));
+				ModeToolbarBuilder.SetStyle(&FEditorStyle::Get(), "PaletteToolBar");
+
+				ModeToolkit->BuildToolPalette(Palette, ModeToolbarBuilder);
+
+				TSharedRef<SWidget> PaletteWidget = ModeToolbarBuilder.MakeWidget();
+
+				PaletteTabBox->AddSlot()
+				[
+					SNew(SCheckBox)
+					.Padding(FMargin(8.f, 4.f, 8.f, 5.f))
+					.Style( FEditorStyle::Get(),  "PaletteToolBar.Tab" )
+					.OnCheckStateChanged_Lambda( [/*PaletteSwitcher, PaletteWidget, */ModeToolkit, Palette] (const ECheckBoxState) { 
+							ModeToolkit->SetCurrentPalette(Palette);
+						} 
+					)
+					// .IsChecked_Lambda( [PaletteSwitcher, PaletteWidget] () -> ECheckBoxState { return PaletteSwitcher->GetActiveWidget() == PaletteWidget ? ECheckBoxState::Checked : ECheckBoxState::Unchecked; })
+					.IsChecked_Lambda( [ModeToolkit, Palette] () -> ECheckBoxState { return ModeToolkit->GetCurrentPalette() == Palette ? ECheckBoxState::Checked : ECheckBoxState::Unchecked; } )
+					[
+						SNew(STextBlock)
+						.TextStyle(FAppStyle::Get(), "NormalText")
+						.Text(ModeToolkit->GetToolPaletteDisplayName(Palette))
+						.Justification(ETextJustify::Center)
+					]
+				];
+
+
+				PaletteSwitcher->AddSlot()
+				[
+					PaletteWidget
+				]; 
+			}
+
+
 			ModeToolHeader->SetContent(
-				SNew(SExpandableArea)
-				.HeaderPadding(FMargin(2.0f))
-				.Padding(FMargin(10.f))
-				.BorderImage(FEditorStyle::Get().GetBrush("DetailsView.CategoryTop"))
-				.BorderBackgroundColor(FLinearColor(.6, .6, .6, 1.0f))
-				.BodyBorderBackgroundColor(FLinearColor::Transparent)
-				.AreaTitleFont(FEditorStyle::Get().GetFontStyle("EditorModesPanel.CategoryFontStyle"))
-				.Visibility_Lambda([ModeToolkit] { return ModeToolkit->GetActiveToolDisplayName().IsEmpty() && ModeToolkit->GetActiveToolMessage().IsEmpty() ? EVisibility::Collapsed : EVisibility::Visible;  })
-				.BodyContent()
+				SNew(SVerticalBox)
+
+				+SVerticalBox::Slot()
+				.Padding(8.f, 0.f, 0.f, 8.f)
+				.AutoHeight()
 				[
-					SNew(STextBlock)
-					.Text(ModeToolkit.Get(), &FModeToolkit::GetActiveToolMessage)
-					.Font(FEditorStyle::Get().GetFontStyle("EditorModesPanel.ToolDescriptionFont"))
-					.AutoWrapText(true)
+					PaletteTabBox
 				]
-				.HeaderContent()
+
+				+SVerticalBox::Slot()
+				.AutoHeight()
 				[
-					SNew(STextBlock)
-					.Text(ModeToolkit.Get(), &FModeToolkit::GetActiveToolDisplayName)
-					.Justification(ETextJustify::Center)
-					.Font(FEditorStyle::Get().GetFontStyle("EditorModesPanel.CategoryFontStyle"))
+					PaletteSwitcher
 				]
 			);
 
diff --git a/Engine/Source/Editor/UnrealEd/Private/EditorModeManager.cpp b/Engine/Source/Editor/UnrealEd/Private/EditorModeManager.cpp
--- a/Engine/Source/Editor/UnrealEd/Private/EditorModeManager.cpp
+++ b/Engine/Source/Editor/UnrealEd/Private/EditorModeManager.cpp
@@ -644,7 +644,7 @@ void FEditorModeTools::RebuildModeToolBar()
 							.Style( FEditorStyle::Get(),  "ToolPalette.DockingTab" )
 							.OnCheckStateChanged_Lambda( [PaletteSwitcher, Row, RowToolkit] (const ECheckBoxState) { 
 									PaletteSwitcher->SetActiveWidget(Row.ToolbarWidget.ToSharedRef());
-									RowToolkit->OnToolPaletteChanged(Row.PaletteName);
+									RowToolkit->SetCurrentPalette(Row.PaletteName);
 								} 
 							)
 							.IsChecked_Lambda( [PaletteSwitcher, PaletteWidget] () -> ECheckBoxState { return PaletteSwitcher->GetActiveWidget() == PaletteWidget ? ECheckBoxState::Checked : ECheckBoxState::Unchecked; })
@@ -697,55 +697,6 @@ void FEditorModeTools::RebuildModeToolBar()
 	}
 }
 
-void FEditorModeTools::SpawnOrUpdateModeToolbar()
-{
-	if(ShouldShowModeToolbar())
-	{
-		if (ModeToolbarTab.IsValid())
-		{
-			RebuildModeToolBar();
-		}
-		else if (ToolkitHost.IsValid())
-		{
-			ToolkitHost.Pin()->GetTabManager()->TryInvokeTab(EditorModeToolbarTabName);
-		}
-	}
-}
-
-void FEditorModeTools::InvokeToolPaletteTab(FEditorModeID InModeID, FName InPaletteName)
-{
-	if (!ModeToolbarPaletteSwitcher.Pin()) 
-	{
-		return;
-	}
-
-	for (auto Row: ActiveToolBarRows)
-	{
-		if (Row.ModeID == InModeID && Row.PaletteName == InPaletteName)
-		{
-			TSharedRef<SWidget> PaletteWidget = Row.ToolbarWidget.ToSharedRef();
-
-			TSharedPtr<FModeToolkit> RowToolkit;
-			if (FEdMode* Mode = GetActiveMode(InModeID))
-			{
-				RowToolkit = Mode->GetToolkit();
-			}
-			else if (UEdMode* ScriptableMode = GetActiveScriptableMode(InModeID))
-			{
-				RowToolkit = ScriptableMode->GetToolkit();
-			}
-
-			TSharedPtr<SWidget> ActiveWidget = ModeToolbarPaletteSwitcher.Pin()->GetActiveWidget();
-			if (RowToolkit && ActiveWidget.Get() != Row.ToolbarWidget.Get())
-			{
-				ModeToolbarPaletteSwitcher.Pin()->SetActiveWidget(Row.ToolbarWidget.ToSharedRef());
-				RowToolkit->OnToolPaletteChanged(Row.PaletteName);
-			}
-			break;	
-		}
-	}	
-}
-
 void FEditorModeTools::RemoveAllDelegateHandlers()
 {
 	if (GEditor)
@@ -954,11 +905,6 @@ void FEditorModeTools::ActivateMode(FEditorModeID InID, bool bToggle)
 			ActiveToolBarRows.Emplace(ScriptableMode->GetID(), Palette, Toolkit->GetToolPaletteDisplayName(Palette), ModeToolbarBuilder.MakeWidget());
 			PaletteCount++;
 		}
-
-		if (PaletteCount > 0)
-		{
-			SpawnOrUpdateModeToolbar();
-		}
 	}
 	
 	// Update the editor UI
diff --git a/Engine/Source/Editor/UnrealEd/Private/Toolkits/BaseToolkit.cpp b/Engine/Source/Editor/UnrealEd/Private/Toolkits/BaseToolkit.cpp
--- a/Engine/Source/Editor/UnrealEd/Private/Toolkits/BaseToolkit.cpp
+++ b/Engine/Source/Editor/UnrealEd/Private/Toolkits/BaseToolkit.cpp
@@ -326,12 +326,16 @@ void FModeToolkit::BuildToolPalette(FName PaletteName, class FToolBarBuilder& To
 	}
 }
 
-void FModeToolkit::OnToolPaletteChanged(FName PaletteName)
+FName FModeToolkit::GetCurrentPalette() const
 {
-	if (UEdMode* EditorMode = GetScriptableEditorMode())
-	{
-		EditorMode->SetCurrentPaletteName(PaletteName);
-	}
+	return CurrentPaletteName;
+}
+
+void FModeToolkit::SetCurrentPalette(FName InPalette)
+{
+	CurrentPaletteName = InPalette;
+	OnToolPaletteChanged(CurrentPaletteName);
+	OnPaletteChangedDelegate.Broadcast(InPalette);
 }
 
 void FModeToolkit::SetModeSettingsObject(UObject* InSettingsObject)
diff --git a/Engine/Source/Editor/UnrealEd/Public/Toolkits/BaseToolkit.h b/Engine/Source/Editor/UnrealEd/Public/Toolkits/BaseToolkit.h
--- a/Engine/Source/Editor/UnrealEd/Public/Toolkits/BaseToolkit.h
+++ b/Engine/Source/Editor/UnrealEd/Public/Toolkits/BaseToolkit.h
@@ -146,8 +146,12 @@ public:
 	virtual FText GetActiveToolDisplayName() const { return FText(); }
 	virtual FText GetActiveToolMessage() const { return FText(); }
 
-	virtual void OnToolPaletteChanged(FName PaletteName);
+	DECLARE_MULTICAST_DELEGATE_OneParam(FOnPaletteChanged, FName);
+	FOnPaletteChanged& OnPaletteChanged() { return OnPaletteChangedDelegate; }
 
+	virtual void OnToolPaletteChanged(FName PaletteName) {};
+	void SetCurrentPalette(FName InName);
+	FName GetCurrentPalette() const;
 	void SetModeSettingsObject(UObject* InSettingsObject);
 
 protected:
@@ -163,4 +167,7 @@ protected:
 	TSharedPtr<SWidget> ToolkitWidget;
 	TSharedPtr<IDetailsView> ModeDetailsView;
 	TSharedPtr<IDetailsView> DetailsView;
+
+	FName CurrentPaletteName;
+	FOnPaletteChanged OnPaletteChangedDelegate;
 };
diff --git a/Engine/Source/Editor/UnrealEd/Public/Tools/UEdMode.h b/Engine/Source/Editor/UnrealEd/Public/Tools/UEdMode.h
--- a/Engine/Source/Editor/UnrealEd/Public/Tools/UEdMode.h
+++ b/Engine/Source/Editor/UnrealEd/Public/Tools/UEdMode.h
@@ -282,23 +282,11 @@ public:
 		return TMap<FName, TArray<TSharedPtr<FUICommandInfo>>>();
 	};
 
-	void SetCurrentPaletteName(FName InName)
-	{
-		CurrentPaletteName = InName;
-		UpdateOnPaletteChange();
-	};
-
-	FName GetCurrentPaletteName() const
-	{
-		return CurrentPaletteName;
-	}
-
 protected:
 	virtual void CreateToolkit();
 	virtual void OnToolStarted(UInteractiveToolManager* Manager, UInteractiveTool* Tool) {};
 	virtual void OnToolEnded(UInteractiveToolManager* Manager, UInteractiveTool* Tool) {};
 	virtual void ActivateDefaultTool() {};
-	virtual void UpdateOnPaletteChange() {};
 
 protected:
 
@@ -314,6 +302,5 @@ protected:
 	UPROPERTY(Transient)
 	UObject* SettingsObject;
 
-	FName CurrentPaletteName;
 };
 
diff --git a/Engine/Source/Runtime/Slate/Private/Framework/MultiBox/MultiBox.cpp b/Engine/Source/Runtime/Slate/Private/Framework/MultiBox/MultiBox.cpp
--- a/Engine/Source/Runtime/Slate/Private/Framework/MultiBox/MultiBox.cpp
+++ b/Engine/Source/Runtime/Slate/Private/Framework/MultiBox/MultiBox.cpp
@@ -581,10 +581,16 @@ EVisibility SMultiBoxWidget::GetCustomizationBorderDragVisibility(const FName In
 	return EVisibility::Collapsed;
 }
 
-void SMultiBoxWidget::AddBlockWidget( const FMultiBlock& Block, TSharedPtr<SHorizontalBox> HorizontalBox, TSharedPtr<SVerticalBox> VerticalBox, EMultiBlockLocation::Type InLocation, bool bSectionContainsIcons )
+void SMultiBoxWidget::AddBlockWidget( const FMultiBlock& Block, TSharedPtr<SHorizontalBox> HorizontalBox, TSharedPtr<SVerticalBox> VerticalBox, EMultiBlockLocation::Type InLocation,  bool bSectionContainsIcons )
 {
 	check( MultiBox.IsValid() );
 
+	// Skip Separators for Uniform Tool Bars
+	if ( Block.GetType() == EMultiBlockType::Separator  && MultiBox->GetType() == EMultiBoxType::UniformToolBar)
+	{
+		return;
+	}
+
 	bool bDisplayExtensionHooks = FMultiBoxSettings::DisplayMultiboxHooks.Get() && Block.GetExtensionHook() != NAME_None;
 
 	TSharedRef<SWidget> BlockWidget = Block.MakeWidget(SharedThis(this), InLocation, bSectionContainsIcons)->AsWidget();
@@ -755,6 +761,7 @@ void SMultiBoxWidget::BuildMultiBoxWidget()
 
 	TSharedPtr< STileView< TSharedPtr<SWidget> > > TileView;
 
+
 	switch (MultiBox->GetType())
 	{
 	case EMultiBoxType::MenuBar:
@@ -776,14 +783,20 @@ void SMultiBoxWidget::BuildMultiBoxWidget()
 		break;
 	case EMultiBoxType::UniformToolBar:
 		{
-			MainWidget = UniformToolbarPanel =
-				SAssignNew(UniformToolbarPanel, SUniformToolbarPanel)
-				.Orientation(Orient_Horizontal)
-				.StyleSet(StyleSet)
-				.StyleName(StyleName)
-				.MinUniformSize(StyleSet->GetFloat(StyleName, ".MinUniformToolbarSize", 0.0f))
-				.MaxUniformSize(StyleSet->GetFloat(StyleName, ".MaxUniformToolbarSize", 0.0f))
-				.OnDropdownOpened(FOnGetContent::CreateSP(this, &SMultiBoxWidget::OnWrapButtonClicked));
+			MainWidget = VerticalBox = SNew (SVerticalBox)
+
+			+SVerticalBox::Slot()
+			.Padding(FMargin(0.f, 2.f))
+			.AutoHeight()
+			[
+				SAssignNew(UniformToolbarPanel, SUniformWrapPanel)
+				.HAlign(HAlign_Left)
+				.MinDesiredSlotWidth(68)
+				.MinDesiredSlotHeight(52)
+				.MaxDesiredSlotWidth(68)
+				.MaxDesiredSlotHeight(52)
+				.SlotPadding( 0.0 )
+			];
 		}
 		break;
 	case EMultiBoxType::ButtonRow:
@@ -982,16 +995,10 @@ TSharedRef<SWidget> SMultiBoxWidget::OnWrapButtonClicked()
 {
 	FMenuBuilder MenuBuilder(true, NULL, TSharedPtr<FExtender>(), false, GetStyleSet());
 	{ 
-		const int32 ClippedIndex = ClippedHorizontalBox.IsValid() ? ClippedHorizontalBox->GetClippedIndex() : UniformToolbarPanel->GetClippedIndex();
-		// Iterate through the array of blocks telling each one to add itself to the menu
 		const TArray< TSharedRef< const FMultiBlock > >& Blocks = MultiBox->GetBlocks();
-		for (int32 BlockIdx = ClippedIndex; BlockIdx < Blocks.Num(); ++BlockIdx)
+		for (int32 BlockIdx = ClippedHorizontalBox->GetClippedIndex(); BlockIdx < Blocks.Num(); ++BlockIdx)
 		{
-			// Skip the first entry if its a separator
-			if (BlockIdx != ClippedIndex || !Blocks[BlockIdx]->IsSeparator())
-			{
-				Blocks[BlockIdx]->CreateMenuEntry(MenuBuilder);
-			}
+			Blocks[BlockIdx]->CreateMenuEntry(MenuBuilder);
 		}
 	}
 
diff --git a/Engine/Source/Runtime/Slate/Private/Widgets/Layout/SUniformWrapPanel.cpp b/Engine/Source/Runtime/Slate/Private/Widgets/Layout/SUniformWrapPanel.cpp
--- a/Engine/Source/Runtime/Slate/Private/Widgets/Layout/SUniformWrapPanel.cpp
+++ b/Engine/Source/Runtime/Slate/Private/Widgets/Layout/SUniformWrapPanel.cpp
@@ -17,6 +17,8 @@ void SUniformWrapPanel::Construct( const FArguments& InArgs )
 	NumRows = 0;
 	MinDesiredSlotWidth = InArgs._MinDesiredSlotWidth.Get();
 	MinDesiredSlotHeight = InArgs._MinDesiredSlotHeight.Get();
+	MaxDesiredSlotWidth = InArgs._MaxDesiredSlotWidth.Get();
+	MaxDesiredSlotHeight = InArgs._MaxDesiredSlotHeight.Get();
 	EvenRowDistribution =  InArgs._EvenRowDistribution.Get();
 	HAlign = InArgs._HAlign.Get();
 
@@ -114,6 +116,9 @@ FVector2D SUniformWrapPanel::ComputeUniformCellSize() const
 	
 	const float CachedMinDesiredSlotWidth = MinDesiredSlotWidth.Get();
 	const float CachedMinDesiredSlotHeight = MinDesiredSlotHeight.Get();
+
+	const float CachedMaxDesiredSlotWidth = MaxDesiredSlotWidth.Get();
+	const float CachedMaxDesiredSlotHeight = MaxDesiredSlotHeight.Get();
 	
 	NumColumns = 0;
 	NumRows = 0;
@@ -135,8 +140,10 @@ FVector2D SUniformWrapPanel::ComputeUniformCellSize() const
 
 				MaxChildDesiredSize.X = FMath::Max( MaxChildDesiredSize.X, ChildDesiredSize.X );
 				MaxChildDesiredSize.Y = FMath::Max( MaxChildDesiredSize.Y, ChildDesiredSize.Y );
-			}
 
+				MaxChildDesiredSize.X = FMath::Min( MaxChildDesiredSize.X, CachedMaxDesiredSlotWidth);
+				MaxChildDesiredSize.Y = FMath::Min( MaxChildDesiredSize.Y, CachedMaxDesiredSlotHeight);
+			}
 		}
 	}
 
@@ -189,6 +196,16 @@ void SUniformWrapPanel::SetMinDesiredSlotHeight(TAttribute<float> InMinDesiredSl
 	MinDesiredSlotHeight = InMinDesiredSlotHeight;
 }
 
+void SUniformWrapPanel::SetMaxDesiredSlotWidth(TAttribute<float> InMaxDesiredSlotWidth)
+{
+	MaxDesiredSlotWidth = InMaxDesiredSlotWidth;
+}
+
+void SUniformWrapPanel::SetMaxDesiredSlotHeight(TAttribute<float> InMaxDesiredSlotHeight)
+{
+	MaxDesiredSlotHeight = InMaxDesiredSlotHeight;
+}
+
 void SUniformWrapPanel::SetHorizontalAlignment(TAttribute<EHorizontalAlignment> InHAlignment)
 {
 	HAlign = InHAlignment;	
diff --git a/Engine/Source/Runtime/Slate/Public/Framework/MultiBox/MultiBox.h b/Engine/Source/Runtime/Slate/Public/Framework/MultiBox/MultiBox.h
--- a/Engine/Source/Runtime/Slate/Public/Framework/MultiBox/MultiBox.h
+++ b/Engine/Source/Runtime/Slate/Public/Framework/MultiBox/MultiBox.h
@@ -11,6 +11,7 @@
 #include "Widgets/Layout/SLinkedBox.h"
 #include "Widgets/SCompoundWidget.h"
 #include "Widgets/Layout/SMenuOwner.h"
+#include "Widgets/Layout/SUniformWrapPanel.h"
 #include "Framework/Commands/UIAction.h"
 #include "Framework/Commands/UICommandInfo.h"
 #include "Framework/Commands/UICommandList.h"
@@ -24,7 +25,7 @@ class SMultiBoxWidget;
 class STableViewBase;
 class SVerticalBox;
 class UToolMenuBase;
-class SUniformToolbarPanel;
+class SUniformWrapPanel;
 
 namespace MultiBoxConstants
 {	
@@ -723,7 +724,7 @@ public:
 
 private:
 	/** Adds a block Widget to this widget */
-	void AddBlockWidget( const FMultiBlock& Block, TSharedPtr<SHorizontalBox> HorizontalBox, TSharedPtr<SVerticalBox> VerticalBox, EMultiBlockLocation::Type InLocation, bool bSectionContainsIcons );
+	void AddBlockWidget( const FMultiBlock& Block, TSharedPtr<SHorizontalBox> HorizontalBox, TSharedPtr<SVerticalBox> VerticalBox, EMultiBlockLocation::Type InLocation,  bool bSectionContainsIcons );
 
 	/**
 	 * Updates the preview block being dragged.  The drag area is where the users dragged block will be dropped
@@ -776,7 +777,7 @@ private:
 	TSharedPtr<class SClippingHorizontalBox> ClippedHorizontalBox;
 
 	/** Specialized box widget to handle clipping of toolbars and menubars */
-	TSharedPtr<SUniformToolbarPanel> UniformToolbarPanel;
+	TSharedPtr<SUniformWrapPanel> UniformToolbarPanel;
 
 	/** A preview of a block being dragged inside this box */
 	FDraggedMultiBlockPreview DragPreview;
diff --git a/Engine/Source/Runtime/Slate/Public/Widgets/Layout/SUniformWrapPanel.h b/Engine/Source/Runtime/Slate/Public/Widgets/Layout/SUniformWrapPanel.h
--- a/Engine/Source/Runtime/Slate/Public/Widgets/Layout/SUniformWrapPanel.h
+++ b/Engine/Source/Runtime/Slate/Public/Widgets/Layout/SUniformWrapPanel.h
@@ -43,6 +43,8 @@ public:
 		: _SlotPadding( FMargin(0.0f) )
 		, _MinDesiredSlotWidth( 0.0f )
 		, _MinDesiredSlotHeight( 0.0f )
+		, _MaxDesiredSlotWidth( FLT_MAX )
+		, _MaxDesiredSlotHeight( FLT_MAX )
 		, _EvenRowDistribution(false)
 		, _HAlign(EHorizontalAlignment::HAlign_Left)
 		{
@@ -61,6 +63,12 @@ public:
 		/** The minimum desired height of the slots */
 		SLATE_ATTRIBUTE(float, MinDesiredSlotHeight)
 
+		/** The minimum desired width of the slots */
+		SLATE_ATTRIBUTE(float, MaxDesiredSlotWidth)
+
+		/** The minimum desired height of the slots */
+		SLATE_ATTRIBUTE(float, MaxDesiredSlotHeight)
+
 		/** If the distribution to evenly distribute down rows */
 		SLATE_ATTRIBUTE(bool, EvenRowDistribution)
 
@@ -85,6 +93,12 @@ public:
 	/** See MinDesiredSlotHeight attribute */
 	void SetMinDesiredSlotHeight(TAttribute<float> InMinDesiredSlotHeight);
 
+	/** See MinDesiredSlotWidth attribute */
+	void SetMaxDesiredSlotWidth(TAttribute<float> InMaxDesiredSlotWidth);
+
+	/** See MinDesiredSlotHeight attribute */
+	void SetMaxDesiredSlotHeight(TAttribute<float> InMaxDesiredSlotHeight);
+
 	/** See HAlign attribute */
 	void SetHorizontalAlignment(TAttribute<EHorizontalAlignment> InHAlignment);
 	EHorizontalAlignment GetHorizontalAlignment() { return HAlign.Get(); }
@@ -125,6 +139,9 @@ private:
 	TAttribute<float> MinDesiredSlotWidth;
 	TAttribute<float> MinDesiredSlotHeight;
 
+	TAttribute<float> MaxDesiredSlotWidth;
+	TAttribute<float> MaxDesiredSlotHeight;
+
 	TAttribute< EHorizontalAlignment > HAlign; 
 	TAttribute<bool> EvenRowDistribution; 
 };
