Export FFractureEditorModeToolkit

diff --git a/Engine/Plugins/Experimental/ChaosEditor/Source/FractureEditor/Public/FractureEditorModeToolkit.h b/Engine/Plugins/Experimental/ChaosEditor/Source/FractureEditor/Public/FractureEditorModeToolkit.h
--- a/Engine/Plugins/Experimental/ChaosEditor/Source/FractureEditor/Public/FractureEditorModeToolkit.h
+++ b/Engine/Plugins/Experimental/ChaosEditor/Source/FractureEditor/Public/FractureEditorModeToolkit.h
@@ -38,7 +38,7 @@ public:
 		FFractureEditorModeToolkit* Toolkit;
 };
 
-class FFractureEditorModeToolkit : public FModeToolkit, public FGCObject
+class FRACTUREEDITOR_API FFractureEditorModeToolkit : public FModeToolkit, public FGCObject
 {
 public:
 
