Fixed crash when creating a geometry collection asset, converting it to Nanite first, and then performing fracturing against it. NaniteData would cache with (for example) 1 resource that mapped 1:1 with the geometry group in the GC, but after fracturing there would be 44 geometry groups but the NaniteData was now stale. The fracture tool forces a proxy recreate which then would crash due to a mismatch between the associated arrays. This change forces the NaniteData to update after the in-place modification of the FGeometryCollection data has completed.

diff --git a/Engine/Plugins/Experimental/ChaosEditor/Source/FractureEditor/Private/FractureEditorModeToolkit.cpp b/Engine/Plugins/Experimental/ChaosEditor/Source/FractureEditor/Private/FractureEditorModeToolkit.cpp
--- a/Engine/Plugins/Experimental/ChaosEditor/Source/FractureEditor/Private/FractureEditorModeToolkit.cpp
+++ b/Engine/Plugins/Experimental/ChaosEditor/Source/FractureEditor/Private/FractureEditorModeToolkit.cpp
@@ -1580,13 +1580,28 @@ void FFractureEditorModeToolkit::ExecuteFracture(FFractureContext& FractureConte
 {
 	if (ActiveTool != nullptr)
 	{
-
 		FractureContext.FracturedGeometryCollection->Modify();
 		ActiveTool->ExecuteFracture(FractureContext);
 		TSharedPtr<FGeometryCollection, ESPMode::ThreadSafe> GeometryCollectionPtr = FractureContext.FracturedGeometryCollection->GetGeometryCollection();
 		FGeometryCollection* OutGeometryCollection = GeometryCollectionPtr.Get();
 		FGeometryCollectionClusteringUtility::UpdateHierarchyLevelOfChildren(OutGeometryCollection, -1);
 
+		// Update Nanite resource data to correctly reflect modified geometry collection data
+		{
+			FractureContext.FracturedGeometryCollection->ReleaseResources();
+
+			if (FractureContext.FracturedGeometryCollection->EnableNanite)
+			{
+				FractureContext.FracturedGeometryCollection->NaniteData = UGeometryCollection::CreateNaniteData(OutGeometryCollection);
+			}
+			else
+			{
+				FractureContext.FracturedGeometryCollection->NaniteData = MakeUnique<FGeometryCollectionNaniteData>();
+			}
+
+			FractureContext.FracturedGeometryCollection->InitResources();
+		}
+
 		UGeometryCollectionComponent* GeometryCollectionComponent = Cast<UGeometryCollectionComponent>(FractureContext.OriginalPrimitiveComponent);
 
 		if (GeometryCollectionComponent != nullptr) // Create new GC actor from static mesh
