tweak fracture UI to avoid invalid or extremely expensive inputs
 - reduce the UI range for random plane cuts; 2000 would create a huge number of cells and seems way too high
 - prevent brick dimensions from going <= 0, which causes the brick layout code to hang

#rnx

diff --git a/Engine/Plugins/Experimental/ChaosEditor/Source/FractureEditor/Private/FractureToolBrick.h b/Engine/Plugins/Experimental/ChaosEditor/Source/FractureEditor/Private/FractureToolBrick.h
--- a/Engine/Plugins/Experimental/ChaosEditor/Source/FractureEditor/Private/FractureToolBrick.h
+++ b/Engine/Plugins/Experimental/ChaosEditor/Source/FractureEditor/Private/FractureToolBrick.h
@@ -68,15 +68,15 @@ public:
 	EFractureBrickProjection Up;
 
 	/** Brick length */
-	UPROPERTY(EditAnywhere, Category = Brick)
+	UPROPERTY(EditAnywhere, Category = Brick, meta = (UIMin = "0.1", UIMax = "500.0", ClampMin = "0.001"))
 	float BrickLength;
 
 	/** Brick Height */
-	UPROPERTY(EditAnywhere, Category = Brick)
+	UPROPERTY(EditAnywhere, Category = Brick, meta = (UIMin = "0.1", UIMax = "500.0", ClampMin = "0.001"))
 	float BrickHeight;
 
 	/** Brick Height */
-	UPROPERTY(EditAnywhere, Category = Brick)
+	UPROPERTY(EditAnywhere, Category = Brick, meta = (UIMin = "0.1", UIMax = "500.0", ClampMin = "0.001"))
 	float BrickDepth;
 
 	UPROPERTY()
diff --git a/Engine/Plugins/Experimental/ChaosEditor/Source/FractureEditor/Private/FractureToolPlaneCut.h b/Engine/Plugins/Experimental/ChaosEditor/Source/FractureEditor/Private/FractureToolPlaneCut.h
--- a/Engine/Plugins/Experimental/ChaosEditor/Source/FractureEditor/Private/FractureToolPlaneCut.h
+++ b/Engine/Plugins/Experimental/ChaosEditor/Source/FractureEditor/Private/FractureToolPlaneCut.h
@@ -21,7 +21,7 @@ public:
 	: NumberPlanarCuts(3) {}
 
 	/** Number of Clusters - Cluster Voronoi Method */
-	UPROPERTY(EditAnywhere, Category = PlaneCut, meta = (DisplayName = "Number of Cuts", UIMin = "1", UIMax = "2000", ClampMin = "1"))
+	UPROPERTY(EditAnywhere, Category = PlaneCut, meta = (DisplayName = "Number of Cuts", UIMin = "1", UIMax = "20", ClampMin = "1"))
 	int32 NumberPlanarCuts;
 
 	/** Actor to be used for voronoi bounds or plane cutting  */
