FractureModeToolkit -- added root node if necessary to geometry collection following a reset.

diff --git a/Engine/Plugins/Experimental/ChaosEditor/Source/FractureEditor/Private/FractureEditorModeToolkit.cpp b/Engine/Plugins/Experimental/ChaosEditor/Source/FractureEditor/Private/FractureEditorModeToolkit.cpp
--- a/Engine/Plugins/Experimental/ChaosEditor/Source/FractureEditor/Private/FractureEditorModeToolkit.cpp
+++ b/Engine/Plugins/Experimental/ChaosEditor/Source/FractureEditor/Private/FractureEditorModeToolkit.cpp
@@ -1559,8 +1559,9 @@ void FFractureEditorModeToolkit::ResetAsset()
 				}
 
 				GeometryCollectionObject->InitializeMaterials();
+
 				FGeometryCollectionClusteringUtility::UpdateHierarchyLevelOfChildren(GeometryCollection, -1);
-				
+				AddSingleRootNodeIfRequired(GeometryCollectionObject);
 				GeometryCollectionComponent->MarkRenderDynamicDataDirty();
 				GeometryCollectionComponent->MarkRenderStateDirty();
 			}
