FractureEditorMode -- AutoCluster tool ignores the GroupFracture flag so that it works with subfracturing.

diff --git a/Engine/Plugins/Experimental/ChaosEditor/Source/FractureEditor/Private/FractureEditorModeToolkit.cpp b/Engine/Plugins/Experimental/ChaosEditor/Source/FractureEditor/Private/FractureEditorModeToolkit.cpp
--- a/Engine/Plugins/Experimental/ChaosEditor/Source/FractureEditor/Private/FractureEditorModeToolkit.cpp
+++ b/Engine/Plugins/Experimental/ChaosEditor/Source/FractureEditor/Private/FractureEditorModeToolkit.cpp
@@ -49,6 +49,7 @@
 #include "GeometryCollection/GeometryCollectionConversion.h"
 
 #include "PlanarCut.h"
+#include "FractureToolAutoCluster.h" 
 #include "SGeometryCollectionOutliner.h"
 #include "FractureSelectionTools.h"
 #include "GeometryCollection/GeometryCollectionAlgo.h"
@@ -984,7 +985,10 @@ void FFractureEditorModeToolkit::GetFractureContexts(TArray<FFractureContext>& F
 					}
 				}
 
-				if (CommonSettings->bGroupFracture)
+				// #todo (bmiller) 
+				// AutoCluster requires that bGroupFracture is true. Long term we should probably look at a  
+				// refactor that doesn't have UFractureToolAutoCluster inherit UFactureTool. 
+				if (CommonSettings->bGroupFracture || Cast<UFractureToolAutoCluster>(CommonSettings->OwnerTool))
 				{
 					FractureContexts.AddDefaulted();
 					FFractureContext& FractureContext = FractureContexts.Last();
