Fracture Tool Mode -- Refactored tools implementation so that a clear hierarchy of tool types is available and the core logic lives in the tool classes instead of the Toolkit class.

FractureToolClusterMagnet -- first pass implementation for artist testing.

diff --git a/Engine/Plugins/Experimental/ChaosEditor/Source/FractureEditor/Private/FractureEditorCommands.cpp b/Engine/Plugins/Experimental/ChaosEditor/Source/FractureEditor/Private/FractureEditorCommands.cpp
--- a/Engine/Plugins/Experimental/ChaosEditor/Source/FractureEditor/Private/FractureEditorCommands.cpp
+++ b/Engine/Plugins/Experimental/ChaosEditor/Source/FractureEditor/Private/FractureEditorCommands.cpp
@@ -16,14 +16,6 @@ FFractureEditorCommands::FFractureEditorCommands()
 void FFractureEditorCommands::RegisterCommands()
 {
 
-	// Selection Commands
-	UI_COMMAND(SelectAll, "Select All", "Selects all Bones in the GeometryCollection", EUserInterfaceActionType::Button, FInputChord(EModifierKey::Control | EModifierKey::Shift, EKeys::A));
-	UI_COMMAND(SelectNone, "Deselect", "Deselects all Bones in the GeometryCollection", EUserInterfaceActionType::Button, FInputChord(EModifierKey::Control | EModifierKey::Shift, EKeys::D));
-	UI_COMMAND(SelectNeighbors, "Neighbors", "Select all bones adjacent to the currently selected bones", EUserInterfaceActionType::Button, FInputChord() );
-	UI_COMMAND(SelectSiblings, "Level", "Select all bones at the same levels as the currently selected bones", EUserInterfaceActionType::Button, FInputChord() );
-	UI_COMMAND(SelectAllInCluster, "Cluster", "Select all bones with the same parent as selected bones", EUserInterfaceActionType::Button, FInputChord() );
-	UI_COMMAND(SelectInvert, "Invert", "Invert the selected bones", EUserInterfaceActionType::Button, FInputChord() );
-
 	// View settings
 	UI_COMMAND(ToggleShowBoneColors, "Colors", "Toggle Show Bone Colors", EUserInterfaceActionType::ToggleButton, FInputChord(EModifierKey::Shift, EKeys::B) );
 	UI_COMMAND(ViewUpOneLevel, "ViewUpOneLevel", "View Up One Level", EUserInterfaceActionType::Button, FInputChord(EModifierKey::Shift, EKeys::W) );
@@ -31,29 +23,17 @@ void FFractureEditorCommands::RegisterCommands()
 	UI_COMMAND(ExplodeMore, "ExplodeMore", "Explode 10% More", EUserInterfaceActionType::Button, FInputChord(EModifierKey::Shift, EKeys::E) );
 	UI_COMMAND(ExplodeLess, "ExplodeLess", "Explode 10% Less", EUserInterfaceActionType::Button, FInputChord(EModifierKey::Shift, EKeys::Q) );
 
-	// Cluster Operations
-	UI_COMMAND(Flatten, "Flatten", "Flattens all bones to level 1", EUserInterfaceActionType::Button, FInputChord() );
-	UI_COMMAND(FlattenToLevel, "FlattenToLevel", "Flatten Only up to the Current View Level", EUserInterfaceActionType::Button, FInputChord() );
-	UI_COMMAND(Cluster, "Cluster", "Clusters selected bones under a new parent", EUserInterfaceActionType::Button, FInputChord() );
-	UI_COMMAND(Uncluster, "Uncluster", "Uncluster", EUserInterfaceActionType::Button, FInputChord() );
-	UI_COMMAND(Merge, "Merge", "Merge", EUserInterfaceActionType::Button, FInputChord() );
-	UI_COMMAND(MoveUp, "Level Up", "Move Bones Up One Level", EUserInterfaceActionType::Button, FInputChord() );
-
-	UI_COMMAND(GenerateAsset, "New", "Generate Geometry Collection Asset from static meshes contained in selected actors", EUserInterfaceActionType::Button, FInputChord() );
-	UI_COMMAND(ResetAsset, "Reset", "Reset Geometry Collection Asset to its initial unfractured state", EUserInterfaceActionType::Button, FInputChord());
-
-
+	// Fracture Tools
 	for (TObjectIterator<UClass> ClassIterator; ClassIterator; ++ClassIterator)
 	{
-		if (ClassIterator->IsChildOf(UFractureTool::StaticClass()) && !ClassIterator->HasAnyClassFlags(CLASS_Abstract | CLASS_Deprecated | CLASS_NewerVersionExists))
+		if (ClassIterator->IsChildOf(UFractureActionTool::StaticClass()) && !ClassIterator->HasAnyClassFlags(CLASS_Abstract | CLASS_Deprecated | CLASS_NewerVersionExists))
 		{
-			TSubclassOf<UFractureTool> SubclassOf = (*ClassIterator);
-			UFractureTool* FractureTool = SubclassOf->GetDefaultObject<UFractureTool>();
+			TSubclassOf<UFractureActionTool> SubclassOf = (*ClassIterator);
+			UFractureActionTool* FractureTool = SubclassOf->GetDefaultObject<UFractureActionTool>();
 			FractureTool->RegisterUICommand(this);
 		}
 	}
 
-
 }
 
 #undef LOCTEXT_NAMESPACE
diff --git a/Engine/Plugins/Experimental/ChaosEditor/Source/FractureEditor/Private/FractureEditorCommands.h b/Engine/Plugins/Experimental/ChaosEditor/Source/FractureEditor/Private/FractureEditorCommands.h
--- a/Engine/Plugins/Experimental/ChaosEditor/Source/FractureEditor/Private/FractureEditorCommands.h
+++ b/Engine/Plugins/Experimental/ChaosEditor/Source/FractureEditor/Private/FractureEditorCommands.h
@@ -11,7 +11,7 @@ class FFractureEditorCommands : public TCommands<FFractureEditorCommands>
 		virtual void RegisterCommands() override;
 
 	public:
-
+		
 		// Selection Commands
 		TSharedPtr< FUICommandInfo > SelectAll;
 		TSharedPtr< FUICommandInfo > SelectNone;
@@ -29,15 +29,16 @@ class FFractureEditorCommands : public TCommands<FFractureEditorCommands>
 
 		// Cluster Commands
 		TSharedPtr< FUICommandInfo > AutoCluster;
+		TSharedPtr< FUICommandInfo > ClusterMagnet;
 		TSharedPtr< FUICommandInfo > Cluster;
 		TSharedPtr< FUICommandInfo > Uncluster;
 		TSharedPtr< FUICommandInfo > Flatten;
-		TSharedPtr< FUICommandInfo > FlattenToLevel;
-		TSharedPtr< FUICommandInfo > Merge;
 		TSharedPtr< FUICommandInfo > MoveUp;
+		
+		// Generate Commands
 		TSharedPtr< FUICommandInfo > GenerateAsset;
 		TSharedPtr< FUICommandInfo > ResetAsset;
-
+		
 		// Fracture Commands
 		TSharedPtr< FUICommandInfo > Uniform;
 		TSharedPtr< FUICommandInfo > Radial;
@@ -46,6 +47,6 @@ class FFractureEditorCommands : public TCommands<FFractureEditorCommands>
 		TSharedPtr< FUICommandInfo > Slice;
 		TSharedPtr< FUICommandInfo > Brick;
 		TSharedPtr< FUICommandInfo > Texture;
-
+		
 };
 
diff --git a/Engine/Plugins/Experimental/ChaosEditor/Source/FractureEditor/Private/FractureEditorMode.cpp b/Engine/Plugins/Experimental/ChaosEditor/Source/FractureEditor/Private/FractureEditorMode.cpp
--- a/Engine/Plugins/Experimental/ChaosEditor/Source/FractureEditor/Private/FractureEditorMode.cpp
+++ b/Engine/Plugins/Experimental/ChaosEditor/Source/FractureEditor/Private/FractureEditorMode.cpp
@@ -60,6 +60,11 @@ void FFractureEditorMode::Enter()
 	TArray<UObject*> SelectedObjects;
 	SelectedActors->GetSelectedObjects(SelectedObjects);
 
+	if (FFractureEditorModeToolkit* FractureToolkit = static_cast<FFractureEditorModeToolkit*>(Toolkit.Get()))
+	{
+		FractureToolkit->SetInitialPalette();
+	}
+	
 	OnActorSelectionChanged(SelectedObjects, false);
 
 }
@@ -120,7 +125,7 @@ void FFractureEditorMode::Render(const FSceneView* View, FViewport* Viewport, FP
 
 	FFractureEditorModeToolkit* FractureToolkit = (FFractureEditorModeToolkit*)Toolkit.Get();
  
-	if (UFractureTool* FractureTool = FractureToolkit->GetActiveTool())
+	if (UFractureModalTool* FractureTool = FractureToolkit->GetActiveTool())
 	{
 		auto Settings = FractureTool->GetSettingsObjects();
 		FractureTool->Render(View, Viewport, PDI);
diff --git a/Engine/Plugins/Experimental/ChaosEditor/Source/FractureEditor/Private/FractureEditorModeToolkit.cpp b/Engine/Plugins/Experimental/ChaosEditor/Source/FractureEditor/Private/FractureEditorModeToolkit.cpp
--- a/Engine/Plugins/Experimental/ChaosEditor/Source/FractureEditor/Private/FractureEditorModeToolkit.cpp
+++ b/Engine/Plugins/Experimental/ChaosEditor/Source/FractureEditor/Private/FractureEditorModeToolkit.cpp
@@ -1,9 +1,11 @@
 // Copyright Epic Games, Inc. All Rights Reserved.
 
+#include "FractureEditorModeToolkit.h"
 #include "FractureEditorModeToolkit.h"
 
 #include "AssetRegistryModule.h"
 #include "EditorModeManager.h"
+
 #include "Engine/Selection.h"
 #include "FractureEditorMode.h"
 #include "Framework/MultiBox/MultiBoxBuilder.h"
@@ -27,14 +29,10 @@
 #include "Widgets/Layout/SGridPanel.h"
 #include "Framework/MultiBox/MultiBoxBuilder.h"
 
-#include "SCreateAssetFromObject.h"
-
 #include "FractureTool.h"
-#include "GeometryCollection/GeometryCollectionComponent.h"
+
 #include "GeometryCollection/GeometryCollectionActor.h"
 #include "GeometryCollection/GeometryCollection.h"
-#include "Engine/StaticMesh.h"
-#include "Engine/SkeletalMesh.h"
 
 #include "Styling/AppStyle.h"
 #include "EditorStyleSet.h"
@@ -43,10 +41,6 @@
 #include "FractureEditorStyle.h"
 #include "ScopedTransaction.h"
 #include "Modules/ModuleManager.h"
-#include "AssetToolsModule.h"
-#include "Layers/LayersSubsystem.h"
-#include "StaticMeshAttributes.h"
-#include "GeometryCollection/GeometryCollectionConversion.h"
 
 #include "PlanarCut.h"
 #include "FractureToolAutoCluster.h" 
@@ -73,7 +67,7 @@ TArray<UClass*> FindFractureToolClasses()
 
 	for (TObjectIterator<UClass> ClassIterator; ClassIterator; ++ClassIterator)
 	{
-		if (ClassIterator->IsChildOf(UFractureTool::StaticClass()) && !ClassIterator->HasAnyClassFlags(CLASS_Abstract | CLASS_Deprecated | CLASS_NewerVersionExists))
+		if (ClassIterator->IsChildOf(UFractureActionTool::StaticClass()) && !ClassIterator->HasAnyClassFlags(CLASS_Abstract | CLASS_Deprecated | CLASS_NewerVersionExists))
 		{
 			Classes.Add(*ClassIterator);
 		}
@@ -206,7 +200,6 @@ void FFractureEditorModeToolkit::Init(const TSharedPtr<IToolkitHost>& InitToolki
 	.BorderBackgroundColor( FLinearColor( .6,.6,.6, 1.0f ) )
 	.BodyBorderBackgroundColor (FLinearColor( 1.0, 0.0, 0.0))
 	.AreaTitleFont(FEditorStyle::Get().GetFontStyle("DetailsView.CategoryFontStyle"))
-	.Visibility_Lambda([this]() -> EVisibility { return (GetActiveTool() == nullptr) ? EVisibility::Visible: EVisibility::Collapsed; })
 	.BodyContent()
 	[
 		SAssignNew(OutlinerView, SGeometryCollectionOutliner )
@@ -221,7 +214,6 @@ void FFractureEditorModeToolkit::Init(const TSharedPtr<IToolkitHost>& InitToolki
 	.BorderBackgroundColor( FLinearColor( .6,.6,.6, 1.0f ) )
 	.BodyBorderBackgroundColor (FLinearColor( 1.0, 0.0, 0.0))
 	.AreaTitleFont(FEditorStyle::Get().GetFontStyle("DetailsView.CategoryFontStyle"))
-	.Visibility_Lambda([this]() -> EVisibility { return (GetActiveTool() == nullptr) ? EVisibility::Visible: EVisibility::Collapsed; })
 	.BodyContent()
 	[
 		SNew(STextBlock)
@@ -271,8 +263,8 @@ void FFractureEditorModeToolkit::Init(const TSharedPtr<IToolkitHost>& InitToolki
 						.TextStyle( FAppStyle::Get(), "ButtonText" )
 						.HAlign(HAlign_Center)
 						.ContentPadding(FMargin(10.f, Padding))
-						.OnClicked(this, &FFractureEditorModeToolkit::OnFractureClicked)
-						.IsEnabled( this, &FFractureEditorModeToolkit::CanExecuteFracture)
+						.OnClicked(this, &FFractureEditorModeToolkit::OnModalClicked)
+						.IsEnabled( this, &FFractureEditorModeToolkit::CanExecuteModal)
 						.Text_Lambda( [this] () -> FText { return ActiveTool ? ActiveTool->GetApplyText().ToUpper() :  LOCTEXT("FractureApplyButton", "APPLY"); })
 						.Visibility_Lambda( [this] () -> EVisibility { return (GetActiveTool() == nullptr) ? EVisibility::Collapsed : EVisibility::Visible; })
 					]
@@ -335,49 +327,60 @@ void FFractureEditorModeToolkit::OnObjectPostEditChange( UObject* Object, FPrope
 	}
 }
 
-const TArray<FName> FFractureEditorModeToolkit::PaletteNames = { FName(TEXT("Fracture")) };
+const TArray<FName> FFractureEditorModeToolkit::PaletteNames = { FName(TEXT("Generate")), FName(TEXT("Select")), FName(TEXT("Fracture")), FName(TEXT("Cluster")) };
 
 FText FFractureEditorModeToolkit::GetToolPaletteDisplayName(FName Palette) const
 { 
 	return FText::FromName(Palette);
 }
 
+void FFractureEditorModeToolkit::SetInitialPalette()
+{
+	// Start in Select Palette if GeometryCollection is selected.
+	if (IsGeometryCollectionSelected())
+	{
+		SetCurrentPalette(TEXT("Select"));
+	}
+	else
+	{
+		SetCurrentPalette(TEXT("Generate"));
+	}
+}
+
 void FFractureEditorModeToolkit::BuildToolPalette(FName PaletteIndex, class FToolBarBuilder& ToolbarBuilder) 
 {
 	const FFractureEditorCommands& Commands = FFractureEditorCommands::Get();
 
-	if (PaletteIndex == PaletteNames[0])
+	if (PaletteIndex == TEXT("Generate"))
 	{
 		ToolbarBuilder.AddToolBarButton(Commands.GenerateAsset);
 		ToolbarBuilder.AddToolBarButton(Commands.ResetAsset);
-
-		ToolbarBuilder.AddSeparator();
+	}
+	else if (PaletteIndex == TEXT("Select"))
+	{
 		ToolbarBuilder.AddToolBarButton(Commands.SelectAll);
 		ToolbarBuilder.AddToolBarButton(Commands.SelectNone);
 		ToolbarBuilder.AddToolBarButton(Commands.SelectNeighbors);
 		ToolbarBuilder.AddToolBarButton(Commands.SelectSiblings);
 		ToolbarBuilder.AddToolBarButton(Commands.SelectAllInCluster);
 		ToolbarBuilder.AddToolBarButton(Commands.SelectInvert);
-
-		ToolbarBuilder.AddSeparator();
-
-		ToolbarBuilder.AddSeparator();
+	}
+	else if (PaletteIndex == TEXT("Fracture"))
+	{
 		ToolbarBuilder.AddToolBarButton(Commands.Uniform);
 		ToolbarBuilder.AddToolBarButton(Commands.Clustered);
 		ToolbarBuilder.AddToolBarButton(Commands.Radial);
 		ToolbarBuilder.AddToolBarButton(Commands.Planar);
 		ToolbarBuilder.AddToolBarButton(Commands.Slice);
 		ToolbarBuilder.AddToolBarButton(Commands.Brick);
-
-
-		ToolbarBuilder.AddSeparator();
-
+	}
+	else if (PaletteIndex == TEXT("Cluster"))
+	{
 		ToolbarBuilder.AddToolBarButton(Commands.AutoCluster);
+		ToolbarBuilder.AddToolBarButton(Commands.ClusterMagnet);
 		ToolbarBuilder.AddToolBarButton(Commands.Flatten);
-		// ToolbarBuilder.AddToolBarButton(Commands.FlattenToLevel);
 		ToolbarBuilder.AddToolBarButton(Commands.Cluster);
 		ToolbarBuilder.AddToolBarButton(Commands.Uncluster);
-		// ToolbarBuilder.AddToolBarButton(Commands.Merge);
 		ToolbarBuilder.AddToolBarButton(Commands.MoveUp);
 	}
 }
@@ -385,138 +388,11 @@ void FFractureEditorModeToolkit::BuildToolPalette(FName PaletteIndex, class FToo
 void FFractureEditorModeToolkit::BindCommands()
 {
 	const FFractureEditorCommands& Commands = FFractureEditorCommands::Get();
-
-	ToolkitCommands->MapAction(
-		Commands.SelectAll,
-		FExecuteAction::CreateSP(
-			this,
-			&FFractureEditorModeToolkit::OnSelectByMode,
-			GeometryCollection::ESelectionMode::AllGeometry
-			),
-		FCanExecuteAction::CreateStatic(&FFractureEditorModeToolkit::IsGeometryCollectionSelected)
-		);
-
-	ToolkitCommands->MapAction(
-		Commands.SelectNone,
-		FExecuteAction::CreateSP(
-			this,
-			&FFractureEditorModeToolkit::OnSelectByMode,
-			GeometryCollection::ESelectionMode::None
-			),
-		FCanExecuteAction::CreateStatic(&FFractureEditorModeToolkit::IsGeometryCollectionSelected)
-		);
-
-	ToolkitCommands->MapAction(
-		Commands.SelectNeighbors,
-		FExecuteAction::CreateSP(
-			this,
-			&FFractureEditorModeToolkit::OnSelectByMode,
-			GeometryCollection::ESelectionMode::Neighbors
-			),
-		FCanExecuteAction::CreateStatic(&FFractureEditorModeToolkit::IsGeometryCollectionSelected)
-		);
-
-	ToolkitCommands->MapAction(
-		Commands.SelectSiblings,
-		FExecuteAction::CreateSP(
-			this,
-			&FFractureEditorModeToolkit::OnSelectByMode,
-			GeometryCollection::ESelectionMode::Siblings
-			),
-		FCanExecuteAction::CreateStatic(&FFractureEditorModeToolkit::IsGeometryCollectionSelected)
-		);
-
-	ToolkitCommands->MapAction(
-		Commands.SelectAllInCluster,
-		FExecuteAction::CreateSP(
-			this,
-			&FFractureEditorModeToolkit::OnSelectByMode,
-			GeometryCollection::ESelectionMode::AllInCluster
-			),
-		FCanExecuteAction::CreateStatic(&FFractureEditorModeToolkit::IsGeometryCollectionSelected)
-		);
-
-	ToolkitCommands->MapAction(
-		Commands.SelectInvert,
-		FExecuteAction::CreateSP(
-			this,
-			&FFractureEditorModeToolkit::OnSelectByMode,
-			GeometryCollection::ESelectionMode::InverseGeometry
-			),
-		FCanExecuteAction::CreateStatic(&FFractureEditorModeToolkit::IsGeometryCollectionSelected)
-		);
-
-	ToolkitCommands->MapAction(
-		Commands.Cluster,
-		FExecuteAction::CreateSP(this, &FFractureEditorModeToolkit::OnCluster),
-		FCanExecuteAction::CreateStatic(&FFractureEditorModeToolkit::IsGeometryCollectionSelected)
-	);
-
-	ToolkitCommands->MapAction(
-		Commands.Uncluster,
-		FExecuteAction::CreateSP(this, &FFractureEditorModeToolkit::OnUncluster),
-		FCanExecuteAction::CreateStatic(&FFractureEditorModeToolkit::IsGeometryCollectionSelected)
-	);
-
-	ToolkitCommands->MapAction(
-		Commands.Flatten,
-		FExecuteAction::CreateSP(this, &FFractureEditorModeToolkit::OnFlatten),
-		FCanExecuteAction::CreateStatic(&FFractureEditorModeToolkit::IsGeometryCollectionSelected)
-	);
-
-	ToolkitCommands->MapAction(
-		Commands.FlattenToLevel,
-		FExecuteAction::CreateSP(this, &FFractureEditorModeToolkit::OnFlattenToLevel),
-		FCanExecuteAction::CreateLambda([] { return false; })
-		);
-
-// 	ToolkitCommands->MapAction(
-// 		Commands.Merge,
-// 		FExecuteAction::CreateSP(this, &FFractureEditorModeToolkit::OnMerge),
-// 		FCanExecuteAction::CreateLambda([] { return false; })
-// 		);
-
-	ToolkitCommands->MapAction(
-		Commands.MoveUp,
-		FExecuteAction::CreateSP(this, &FFractureEditorModeToolkit::OnMoveUp),
-		FCanExecuteAction::CreateLambda([] { return IsGeometryCollectionSelected(); })
-	);
-
-	ToolkitCommands->MapAction(
-		Commands.GenerateAsset,
-		FExecuteAction::CreateSP(this, &FFractureEditorModeToolkit::GenerateAsset),
-		FCanExecuteAction::CreateLambda([] { return IsStaticMeshSelected(); })
-		);
-
-	ToolkitCommands->MapAction(
-		Commands.ResetAsset,
-		FExecuteAction::CreateSP(this, &FFractureEditorModeToolkit::ResetAsset),
-		FCanExecuteAction::CreateStatic(&FFractureEditorModeToolkit::IsGeometryCollectionSelected)
-	);
-
-
-	TArray<UClass*> SourceClasses = FindFractureToolClasses();
-	for (UClass* Class : SourceClasses)
-	{
-		TSubclassOf<UFractureTool> SubclassOf = Class;
-		UFractureTool* FractureTool = SubclassOf->GetDefaultObject<UFractureTool>();
-
-		// Only Bind Commands With Legitmately Set Commands
-		if (FractureTool->GetUICommandInfo())
-		{
-			ToolkitCommands->MapAction(
-				FractureTool->GetUICommandInfo(),
-				FExecuteAction::CreateSP(this, &FFractureEditorModeToolkit::SetActiveTool, FractureTool),
-				FCanExecuteAction(),//::CreateLambda([]() { return true; }),
-				FIsActionChecked::CreateSP(this, &FFractureEditorModeToolkit::IsActiveTool, FractureTool)
-			);
-		}
-	}
-
+	
 	ToolkitCommands->MapAction(
 		Commands.ToggleShowBoneColors,
 		FExecuteAction::CreateSP(this, &FFractureEditorModeToolkit::OnToggleShowBoneColors),
-		FCanExecuteAction(),//::CreateLambda([]() { return true; }),
+		FCanExecuteAction(),
 		FIsActionChecked::CreateSP(this, &FFractureEditorModeToolkit::GetShowBoneColors)
 	);
 
@@ -536,12 +412,49 @@ void FFractureEditorModeToolkit::BindCommands()
 		EUIActionRepeatMode::RepeatEnabled
 	);
 
-
 	ToolkitCommands->MapAction(
 		Commands.ExplodeLess,
 		FExecuteAction::CreateLambda([=]() { this->OnSetExplodedViewValue( FMath::Max(0.0, this->GetExplodedViewValue() - .1) ); } ),
 		EUIActionRepeatMode::RepeatEnabled
 	);
+
+	// Map actions of all the Fracture Tools
+	TArray<UClass*> SourceClasses = FindFractureToolClasses();
+	for (UClass* Class : SourceClasses)
+	{
+		if (Class->IsChildOf(UFractureModalTool::StaticClass()))
+		{
+			TSubclassOf<UFractureModalTool> SubclassOf = Class;
+			UFractureModalTool* FractureTool = SubclassOf->GetDefaultObject<UFractureModalTool>();
+
+			// Only Bind Commands With Legitimately Set Commands
+			if (FractureTool->GetUICommandInfo())
+			{
+				ToolkitCommands->MapAction(
+					FractureTool->GetUICommandInfo(),
+					FExecuteAction::CreateSP(this, &FFractureEditorModeToolkit::SetActiveTool, FractureTool),
+					FCanExecuteAction(),
+					FIsActionChecked::CreateSP(this, &FFractureEditorModeToolkit::IsActiveTool, FractureTool)
+				);
+			}
+		}
+		else
+		{
+			TSubclassOf<UFractureActionTool> SubclassOf = Class;
+			UFractureActionTool* FractureTool = SubclassOf->GetDefaultObject<UFractureActionTool>();
+
+			// Only Bind Commands With Legitimately Set Commands
+			if (FractureTool->GetUICommandInfo())
+			{
+				ToolkitCommands->MapAction(
+					FractureTool->GetUICommandInfo(),
+					FExecuteAction::CreateSP(this, &FFractureEditorModeToolkit::ExecuteAction, FractureTool),
+					FCanExecuteAction::CreateSP(this, &FFractureEditorModeToolkit::CanExecuteAction, FractureTool)
+				);
+			}
+		}
+
+	}
 }
 
 void FFractureEditorModeToolkit::OnToolPaletteChanged(FName PaletteName) 
@@ -798,12 +711,32 @@ TSharedRef<SWidget> FFractureEditorModeToolkit::GetViewMenuContent()
 	return MenuBuilder.MakeWidget();
 }
 
-void FFractureEditorModeToolkit::SetActiveTool(UFractureTool* InActiveTool)
+void FFractureEditorModeToolkit::ExecuteAction(UFractureActionTool* InActionTool)
+{
+	if (InActionTool)
+	{
+		InActionTool->Execute(StaticCastSharedRef<FFractureEditorModeToolkit>(AsShared()));
+	}
+}
+
+bool FFractureEditorModeToolkit::CanExecuteAction(UFractureActionTool* InActionTool) const
+{
+	if (InActionTool)
+	{
+		return InActionTool->CanExecute();
+	}
+	else
+	{
+		return false;
+	}
+}
+
+void FFractureEditorModeToolkit::SetActiveTool(UFractureModalTool* InActiveTool)
 {
 	ActiveTool = InActiveTool;
 
-	UFractureCommonSettings* CommonSettings = GetMutableDefault<UFractureCommonSettings>();
-	CommonSettings->OwnerTool = ActiveTool;
+	UFractureToolSettings* ToolSettings = GetMutableDefault<UFractureToolSettings>();
+	ToolSettings->OwnerTool = ActiveTool;
 
 	TArray<UObject*> Settings;
 	Settings.Add(GetMutableDefault<UFractureSettings>());
@@ -819,12 +752,12 @@ void FFractureEditorModeToolkit::SetActiveTool(UFractureTool* InActiveTool)
 }
 
 
-UFractureTool* FFractureEditorModeToolkit::GetActiveTool() const
+UFractureModalTool* FFractureEditorModeToolkit::GetActiveTool() const
 {
 	return ActiveTool;
 }
 
-bool FFractureEditorModeToolkit::IsActiveTool(UFractureTool* InActiveTool)
+bool FFractureEditorModeToolkit::IsActiveTool(UFractureModalTool* InActiveTool)
 {
 	return bool(ActiveTool == InActiveTool);
 }
@@ -872,20 +805,6 @@ void FFractureEditorModeToolkit::SetOutlinerComponents(const TArray<UGeometryCol
 	{
 		ActiveTool->FractureContextChanged();
 	}
-
-	// Make sure all these selected componenets are set to view the correct viewing level
-
-
-	//if (InNewComponents)
-	//{
-		//GeometryCollection = NewGeometryCollection;
-
-	//}
-	/*else
-	{
-		// Possibly not necessary here... have to get notification for when the collection changed.
-		OutlinerView->UpdateGeometryCollection();
-	}*/
 }
 
 void FFractureEditorModeToolkit::SetBoneSelection(UGeometryCollectionComponent* InRootComponent, const TArray<int32>& InSelectedBones, bool bClearCurrentSelection)
@@ -897,214 +816,26 @@ void FFractureEditorModeToolkit::SetBoneSelection(UGeometryCollectionComponent*
 	}
 }
 
-static bool CanFractureComponent(UPrimitiveComponent* PrimitiveComponent)
-{
-	// Don't bother with editor-only 'helper' actors, we never want to visualize or edit geometry on those
-	return !PrimitiveComponent->IsEditorOnly() &&
-			PrimitiveComponent->GetCollisionEnabled() != ECollisionEnabled::NoCollision &&
-			(PrimitiveComponent->GetOwner() == nullptr || !PrimitiveComponent->GetOwner()->IsEditorOnly());
-}
-
-void FFractureEditorModeToolkit::OnSelectByMode(GeometryCollection::ESelectionMode SelectionMode)
-{
-	USelection* SelectionSet = GEditor->GetSelectedActors();
-
-	TArray<AActor*> SelectedActors;
-	SelectedActors.Reserve(SelectionSet->Num());
-	SelectionSet->GetSelectedObjects(SelectedActors);
-
-	for (AActor* Actor : SelectedActors)
-	{
-		TInlineComponentArray<UGeometryCollectionComponent*> GeometryCollectionComponents;
-		Actor->GetComponents(GeometryCollectionComponents);
-
-		for (UGeometryCollectionComponent* GeometryCollectionComponent : GeometryCollectionComponents)
-		{
-			FScopedColorEdit EditBoneColor = GeometryCollectionComponent->EditBoneSelection();
-			EditBoneColor.SelectBones(SelectionMode);
-			SetBoneSelection(GeometryCollectionComponent, EditBoneColor.GetSelectedBones(), true);
-		}
-	}
-}
-
-void FFractureEditorModeToolkit::GetFractureContexts(TArray<FFractureContext>& FractureContexts)
-{
-	const UFractureCommonSettings* CommonSettings = GetDefault<UFractureCommonSettings>();
-	FRandomStream RandomStream(CommonSettings->RandomSeed > -1 ? CommonSettings->RandomSeed : FMath::Rand());
-
-	USelection* SelectedActors = GEditor->GetSelectedActors();
-	for (FSelectionIterator Iter(*SelectedActors); Iter; ++Iter)
-	{
-		AActor* Actor = Cast<AActor>(*Iter);
-		if (Actor)
-		{
-			TInlineComponentArray<UGeometryCollectionComponent*> GeometryCollectionComponents;
-			Actor->GetComponents(GeometryCollectionComponents);
-			for (UGeometryCollectionComponent* GeometryCollectionComponent : GeometryCollectionComponents)
-			{
-				FGeometryCollectionEdit RestCollection = GeometryCollectionComponent->EditRestCollection(GeometryCollection::EEditUpdate::None);
-				UGeometryCollection* FracturedGeometryCollection = RestCollection.GetRestCollection();
-				if (FracturedGeometryCollection == nullptr)
-				{
-					continue;
-				}
-
-				const TArray<int32>& SelectedBonesOriginal = GeometryCollectionComponent->GetSelectedBones();
-
-				TSharedPtr<FGeometryCollection, ESPMode::ThreadSafe> GeometryCollectionPtr = FracturedGeometryCollection->GetGeometryCollection();
-				FGeometryCollection* OutGeometryCollection = GeometryCollectionPtr.Get();
-
-				const TManagedArray<TSet<int32>>& Children = OutGeometryCollection->GetAttribute<TSet<int32>>("Children", FGeometryCollection::TransformGroup);
-
-				TArray<int32> SelectedBones;
-				SelectedBones.Reserve(SelectedBonesOriginal.Num());
-				for (int32 BoneIndex : SelectedBonesOriginal)
-				{
-					if (Children[BoneIndex].Num() == 0)
-					{
-						SelectedBones.Add(BoneIndex);
-					}
-				}
-
-				const TManagedArray<FTransform>& Transform = OutGeometryCollection->GetAttribute<FTransform>("Transform", FGeometryCollection::TransformGroup);
-				const TManagedArray<int32>& TransformToGeometryIndex = OutGeometryCollection->GetAttribute<int32>("TransformToGeometryIndex", FGeometryCollection::TransformGroup);
-
-				const TManagedArray<FBox>& BoundingBoxes = OutGeometryCollection->GetAttribute<FBox>("BoundingBox", FGeometryCollection::GeometryGroup);
-
-				TArray<FTransform> Transforms;
-				GeometryCollectionAlgo::GlobalMatrices(Transform, OutGeometryCollection->Parent, Transforms);
-
-				TMap<int32, FBox> BoundsToBone;
-
-				for (int32 Idx = 0, ni = FracturedGeometryCollection->NumElements(FGeometryCollection::TransformGroup); Idx < ni; ++Idx)
-				{
-					if (TransformToGeometryIndex[Idx] > -1)
-					{
-						ensure(TransformToGeometryIndex[Idx] > -1);
-						BoundsToBone.Add(Idx, BoundingBoxes[TransformToGeometryIndex[Idx]].TransformBy(Transforms[Idx]));
-					}
-				}
-
-				// #todo (bmiller) 
-				// AutoCluster requires that bGroupFracture is true. Long term we should probably look at a  
-				// refactor that doesn't have UFractureToolAutoCluster inherit UFactureTool. 
-				if (CommonSettings->bGroupFracture || Cast<UFractureToolAutoCluster>(CommonSettings->OwnerTool))
-				{
-					FractureContexts.AddDefaulted();
-					FFractureContext& FractureContext = FractureContexts.Last();
-					FractureContext.RandomSeed = FMath::Rand();
-					if (CommonSettings->RandomSeed > -1)
-					{
-						// make sure it's unique for each context if it's specified.
-						FractureContext.RandomSeed = CommonSettings->RandomSeed + FractureContexts.Num();
-					}
-
-					FractureContext.OriginalActor = Actor;
-					FractureContext.Transform = Actor->GetActorTransform();
-					FractureContext.OriginalPrimitiveComponent = GeometryCollectionComponent;
-					FractureContext.FracturedGeometryCollection = FracturedGeometryCollection;
-					FractureContext.SelectedBones = SelectedBones;
-
-					FractureContext.Bounds = FBox(ForceInit);
-					for (int32 BoneIndex : FractureContext.SelectedBones)
-					{
-						if (FractureContext.SelectedBones.Num() > 1 && RandomStream.FRand() > CommonSettings->ChanceToFracture)
-						{
-							continue;
-						}
-
-						if (TransformToGeometryIndex[BoneIndex] > -1)
-						{
-							FractureContext.Bounds += BoundsToBone[BoneIndex];
-						}
-					}
-				}
-				else
-				{
-					for (int32 BoneIndex : SelectedBones)
-					{
-						if (SelectedBones.Num() > 1 && RandomStream.FRand() > CommonSettings->ChanceToFracture)
-						{
-							continue;
-						}
-
-						FractureContexts.AddDefaulted();
-						FFractureContext& FractureContext = FractureContexts.Last();
-						FractureContext.RandomSeed = FMath::Rand();
-						if (CommonSettings->RandomSeed > -1)
-						{
-							// make sure it's unique for each context if it's specified.
-							FractureContext.RandomSeed = CommonSettings->RandomSeed + FractureContexts.Num();
-						}
-
-						FractureContext.OriginalActor = Actor;
-						FractureContext.Transform = Actor->GetActorTransform();
-						FractureContext.OriginalPrimitiveComponent = GeometryCollectionComponent;
-						FractureContext.FracturedGeometryCollection = FracturedGeometryCollection;
-						FractureContext.SelectedBones = { BoneIndex };
-						if(TransformToGeometryIndex[BoneIndex] > -1)
-						{
-							FractureContext.Bounds = BoundsToBone[BoneIndex];
-						}
-					}
-				}
-			}
-		}
-	}
-}
-
-FReply FFractureEditorModeToolkit::OnFractureClicked()
+FReply FFractureEditorModeToolkit::OnModalClicked()
 {
 	if (ActiveTool)
 	{
 		const double CacheStartTime = FPlatformTime::Seconds();
 
-		TArray<FFractureContext> FractureContexts;
-		GetFractureContexts(FractureContexts);
-
 		FScopedTransaction Transaction(LOCTEXT("FractureMesh", "Fracture Mesh"));
 
-		TArray<UGeometryCollectionComponent*> NewComponents;
-
-		for (FFractureContext& FractureContext : FractureContexts)
-		{
-			ExecuteFracture(FractureContext);
-			UGeometryCollectionComponent* GeometryCollectionComponent = Cast<UGeometryCollectionComponent>(FractureContext.OriginalPrimitiveComponent);
-			NewComponents.AddUnique(GeometryCollectionComponent);
-		}
-
-
-		for (UGeometryCollectionComponent* GeometryCollectionComponent : NewComponents)
-		{
-			FGeometryCollectionEdit GCEdit = GeometryCollectionComponent->EditRestCollection();
-			UGeometryCollection* GCObject = GCEdit.GetRestCollection();
-			TSharedPtr<FGeometryCollection, ESPMode::ThreadSafe> GeometryCollectionPtr = GCObject->GetGeometryCollection();
-			FGeometryCollectionClusteringUtility::UpdateHierarchyLevelOfChildren(GeometryCollectionPtr.Get(), -1);
-
-			FScopedColorEdit EditBoneColor(GeometryCollectionComponent, true);
-
-			EditBoneColor.SelectBones(GeometryCollection::ESelectionMode::None);
-			SetBoneSelection(GeometryCollectionComponent, EditBoneColor.GetSelectedBones(), true);
-
-			UpdateExplodedVectors(GeometryCollectionComponent);
-
-			GeometryCollectionComponent->MarkRenderDynamicDataDirty();
-			GeometryCollectionComponent->MarkRenderStateDirty();
-		}
-
-		SetOutlinerComponents(NewComponents);
-
-		SetActiveTool(0);
+		ActiveTool->Execute(StaticCastSharedRef<FFractureEditorModeToolkit>(AsShared()));
 
 		float ProcessingTime = static_cast<float>(FPlatformTime::Seconds() - CacheStartTime);
 
 		GCurrentLevelEditingViewportClient->Invalidate();
+
 	}
 
 	return FReply::Handled();
 }
 
-bool FFractureEditorModeToolkit::CanExecuteFracture() const
+bool FFractureEditorModeToolkit::CanExecuteModal() const
 {
 
 	if (!IsSelectedActorsInEditorWorld())
@@ -1112,73 +843,14 @@ bool FFractureEditorModeToolkit::CanExecuteFracture() const
 		return false;
 	}
 
-	if (!IsGeometryCollectionSelected())
-	{
-		return false;
-	}
-
-	if (IsStaticMeshSelected())
-	{
-		return false;
-	}
-
 	if (ActiveTool != nullptr) 
 	{
-		return ActiveTool->CanExecuteFracture();
+		return ActiveTool->CanExecute();
 	}
 	
 	return false;
 }
 
-
-bool FFractureEditorModeToolkit::IsLeafBoneSelected() 
-{
-	TSet<UGeometryCollectionComponent*> GeomCompSelection;
-	GetSelectedGeometryCollectionComponents(GeomCompSelection);
-	for (UGeometryCollectionComponent* GeometryCollectionComponent : GeomCompSelection)
-	{
-		const TArray<int32> SelectedBones = GeometryCollectionComponent->GetSelectedBones();
-
-		if (SelectedBones.Num() > 0)
-		{
-			if (const UGeometryCollection* GCObject = GeometryCollectionComponent->GetRestCollection())
-			{
-				TSharedPtr<FGeometryCollection, ESPMode::ThreadSafe> GeometryCollectionPtr = GCObject->GetGeometryCollection();
-				if (const FGeometryCollection* GeometryCollection = GeometryCollectionPtr.Get())
-				{
-					const TManagedArray<TSet<int32>>& Children = GeometryCollection->GetAttribute<TSet<int32>>("Children", FGeometryCollection::TransformGroup);
-
-					for (int32 BoneIndex : SelectedBones)
-					{
-						if (Children[BoneIndex].Num() == 0)
-						{
-							return true;
-						}
-					}
-				}
-			}
-		}
-	}
-	return false;
-}
-
-ULevel* FFractureEditorModeToolkit::GetSelectedLevel()
-{
-	USelection* SelectedActors = GEditor->GetSelectedActors();
-	TArray<ULevel*> UniqueLevels;
-	for (FSelectionIterator Iter(*SelectedActors); Iter; ++Iter)
-	{
-		AActor* Actor = Cast<AActor>(*Iter);
-		if (Actor)
-		{
-			UniqueLevels.AddUnique(Actor->GetLevel());
-		}
-	}
-	check(UniqueLevels.Num() == 1);
-	return UniqueLevels[0];
-}
-
-
 void FFractureEditorModeToolkit::GetSelectedGeometryCollectionComponents(TSet<UGeometryCollectionComponent*>& GeomCompSelection)
 {
 	USelection* SelectionSet = GEditor->GetSelectedActors();
@@ -1196,102 +868,6 @@ void FFractureEditorModeToolkit::GetSelectedGeometryCollectionComponents(TSet<UG
 	}
 }
 
-void FFractureEditorModeToolkit::OnCluster()
-{
-	TSet<UGeometryCollectionComponent*> GeomCompSelection;
-	GetSelectedGeometryCollectionComponents(GeomCompSelection);
-	for (UGeometryCollectionComponent* GeometryCollectionComponent : GeomCompSelection)
-	{
-		const TArray<int32> SelectedBones = GeometryCollectionComponent->GetSelectedBones();
-
-		if (SelectedBones.Num() > 1)
-		{
-			FGeometryCollectionEdit GCEdit = GeometryCollectionComponent->EditRestCollection();
-			if (UGeometryCollection* GCObject = GCEdit.GetRestCollection())
-			{
-				TSharedPtr<FGeometryCollection, ESPMode::ThreadSafe> GeometryCollectionPtr = GCObject->GetGeometryCollection();
-				if (FGeometryCollection* GeometryCollection = GeometryCollectionPtr.Get())
-				{
-					const TManagedArray<TSet<int32>>& Children = GeometryCollection->GetAttribute<TSet<int32>>("Children", FGeometryCollection::TransformGroup);
-
-					// sort the selection list so ClusterBonesUnderNewNode() happens in the correct order for leaf nodes
-					TArray<int32> SortedSelectedBones;
-					SortedSelectedBones.Reserve(SelectedBones.Num());
-					for (int32 SelectedBone : SelectedBones)
-					{
-						if(Children[SelectedBone].Num() > 0)
-						{
-							SortedSelectedBones.Insert(SelectedBone, 0);
-						}
-						else
-						{
-							SortedSelectedBones.Add(SelectedBone);
-						}
-					}
-					// cluster Selected Bones under the first selected bone
-					int32 InsertAtIndex = SortedSelectedBones[0];
-
-					FGeometryCollectionClusteringUtility::ClusterBonesUnderNewNode(GeometryCollection, InsertAtIndex, SortedSelectedBones, false);
-					FGeometryCollectionClusteringUtility::UpdateHierarchyLevelOfChildren(GeometryCollection, -1);
-
-					FScopedColorEdit EditBoneColor = GeometryCollectionComponent->EditBoneSelection();
-					EditBoneColor.ResetBoneSelection();
-					EditBoneColor.ResetHighlightedBones();
-					GeometryCollectionComponent->MarkRenderDynamicDataDirty();
-					GeometryCollectionComponent->MarkRenderStateDirty();
-					SetBoneSelection(GeometryCollectionComponent, EditBoneColor.GetSelectedBones(), true);
-				}
-			}
-		}
-	}
-
-	SetOutlinerComponents(GeomCompSelection.Array());
-}
-
-
-
-void FFractureEditorModeToolkit::OnUncluster()
-{
-	int32 FractureLevel = GetLevelViewValue();
-	TSet<UGeometryCollectionComponent*> GeomCompSelection;
-	GetSelectedGeometryCollectionComponents(GeomCompSelection);
-	for (UGeometryCollectionComponent* GeometryCollectionComponent : GeomCompSelection)
-	{
-		// scoped edit of collection
-		FGeometryCollectionEdit GeometryCollectionEdit = GeometryCollectionComponent->EditRestCollection();
-		if (UGeometryCollection* GeometryCollectionObject = GeometryCollectionEdit.GetRestCollection())
-		{
-			TSharedPtr<FGeometryCollection, ESPMode::ThreadSafe> GeometryCollectionPtr = GeometryCollectionObject->GetGeometryCollection();
-			if (FGeometryCollection* GeometryCollection = GeometryCollectionPtr.Get())
-			{
-				FGeometryCollectionClusteringUtility::UpdateHierarchyLevelOfChildren(GeometryCollection, -1);
-				FGeometryCollectionClusteringUtility::CollapseSelectedHierarchy(FractureLevel, GeometryCollectionComponent->GetSelectedBones(), GeometryCollection);
-
-				FScopedColorEdit EditBoneColor = GeometryCollectionComponent->EditBoneSelection();
-				EditBoneColor.ResetBoneSelection();
-				EditBoneColor.ResetHighlightedBones();
-				GeometryCollectionComponent->MarkRenderDynamicDataDirty();
-				GeometryCollectionComponent->MarkRenderStateDirty();
-				SetBoneSelection(GeometryCollectionComponent, EditBoneColor.GetSelectedBones(), true);
-			}
-		}
-	}
-	SetOutlinerComponents(GeomCompSelection.Array());
-
-}
-
-void FFractureEditorModeToolkit::AddSingleRootNodeIfRequired(UGeometryCollection* GeometryCollectionObject)
-{
-	TSharedPtr<FGeometryCollection, ESPMode::ThreadSafe> GeometryCollectionPtr = GeometryCollectionObject->GetGeometryCollection();
-	if (FGeometryCollection* GeometryCollection = GeometryCollectionPtr.Get())
-	{
-		if (FGeometryCollectionClusteringUtility::ContainsMultipleRootBones(GeometryCollection))
-		{
-			FGeometryCollectionClusteringUtility::ClusterAllBonesUnderNewRoot(GeometryCollection);
-		}
-	}
-}
-
 void FFractureEditorModeToolkit::AddAdditionalAttributesIfRequired(UGeometryCollection* GeometryCollectionObject)
 {
 	TSharedPtr<FGeometryCollection, ESPMode::ThreadSafe> GeometryCollectionPtr = GeometryCollectionObject->GetGeometryCollection();
@@ -1304,328 +880,6 @@ void FFractureEditorModeToolkit::AddAdditionalAttributesIfRequired(UGeometryColl
 	}
 }
 
-
-void FFractureEditorModeToolkit::OnFlatten()
-{
-	TSet<UGeometryCollectionComponent*> GeomCompSelection;
-	GetSelectedGeometryCollectionComponents(GeomCompSelection);
-	for (UGeometryCollectionComponent* GeometryCollectionComponent : GeomCompSelection)
-	{
-		FGeometryCollectionEdit GCEdit = GeometryCollectionComponent->EditRestCollection();
-		if (UGeometryCollection* GCObject = GCEdit.GetRestCollection())
-		{
-			TSharedPtr<FGeometryCollection, ESPMode::ThreadSafe> GeometryCollectionPtr = GCObject->GetGeometryCollection();
-			if (FGeometryCollection* GeometryCollection = GeometryCollectionPtr.Get())
-			{
-				AddAdditionalAttributesIfRequired(GCObject);
-				AddSingleRootNodeIfRequired(GCObject);
-
-				int32 NumElements = GCObject->NumElements(FGeometryCollection::TransformGroup);
-				TArray<int32> Elements;
-				Elements.Reserve(NumElements);
-
-				for (int32 Element = 0; Element < NumElements; ++Element)
-				{
-					if (GeometryCollection->Parent[Element] != FGeometryCollection::Invalid)
-					{
-						Elements.Add(Element);
-					}
-				}
-
-				if (Elements.Num() > 0)
-				{
-					FGeometryCollectionClusteringUtility::ClusterBonesUnderExistingRoot(GeometryCollection, Elements);
-				}
-
-				FGeometryCollectionClusteringUtility::UpdateHierarchyLevelOfChildren(GeometryCollection, -1);
-
-				FScopedColorEdit EditBoneColor = GeometryCollectionComponent->EditBoneSelection();
-				EditBoneColor.ResetBoneSelection();
-
-				OnSetLevelViewValue(1);
-
-				GeometryCollectionComponent->MarkRenderDynamicDataDirty();
-				GeometryCollectionComponent->MarkRenderStateDirty();
-			}
-		}
-	}
-
-	SetOutlinerComponents(GeomCompSelection.Array());
-}
-
-
-void FFractureEditorModeToolkit::OnMoveUp()
-{
-	TSet<UGeometryCollectionComponent*> GeomCompSelection;
-	GetSelectedGeometryCollectionComponents(GeomCompSelection);
-	for (UGeometryCollectionComponent* GeometryCollectionComponent : GeomCompSelection)
-	{
-		FGeometryCollectionEdit GeometryCollectionEdit = GeometryCollectionComponent->EditRestCollection();
-		if (UGeometryCollection* GeometryCollectionObject = GeometryCollectionEdit.GetRestCollection())
-		{
-			TSharedPtr<FGeometryCollection, ESPMode::ThreadSafe> GeometryCollectionPtr = GeometryCollectionObject->GetGeometryCollection();
-			if (FGeometryCollection* GeometryCollection = GeometryCollectionPtr.Get())
-			{
-				TArray<int32> Selected = GeometryCollectionComponent->GetSelectedBones();
-				FGeometryCollectionClusteringUtility::MoveUpOneHierarchyLevel(GeometryCollection, Selected);
-
-				GeometryCollectionComponent->MarkRenderDynamicDataDirty();
-				GeometryCollectionComponent->MarkRenderStateDirty();
-			}
-		}
-	}
-	SetOutlinerComponents(GeomCompSelection.Array());
-}
-
-void FFractureEditorModeToolkit::GenerateAsset()
-{
-	USelection* SelectionSet = GEditor->GetSelectedActors();
-
-	TArray<AActor*> SelectedActors;
-	SelectedActors.Reserve(SelectionSet->Num());
-
-	FScopedTransaction Transaction(LOCTEXT("GenerateAsset", "Generate Geometry Collection Asset"));
-
-	SelectionSet->GetSelectedObjects(SelectedActors);
-
-	OpenGenerateAssetDialog(SelectedActors);
-}
-
-void FFractureEditorModeToolkit::OpenGenerateAssetDialog(TArray<AActor*>& Actors)
-{
-	TSharedPtr<SWindow> PickAssetPathWindow;
-
-	SAssignNew(PickAssetPathWindow, SWindow)
-	.Title(LOCTEXT("SelectPath", "Select Path"))
-	.ToolTipText(LOCTEXT("SelectPathTooltip", "Select the path where the Geometry Collection will be created at"))
-	.ClientSize(FVector2D(400, 400));
-
-	// NOTE - the parent window has to completely exist before this one does so the parent gets set properly.
-	// This is why we do not just put this in the Contents()[ ... ] of the Window above.
-	TSharedPtr<SCreateAssetFromObject> CreateAssetDialog;
-	PickAssetPathWindow->SetContent(
-		SAssignNew(CreateAssetDialog, SCreateAssetFromObject, PickAssetPathWindow)
-		.AssetFilenameSuffix(TEXT("GeometryCollection"))
-		.HeadingText(LOCTEXT("CreateGeometryCollection_Heading", "Geometry Collection Name"))
-		.CreateButtonText(LOCTEXT("CreateGeometryCollection_ButtonLabel", "Create Geometry Collection"))
-		.OnCreateAssetAction(FOnPathChosen::CreateSP(this, &FFractureEditorModeToolkit::OnGenerateAssetPathChosen, Actors))
-	);
-
-	TSharedPtr<SWindow> RootWindow = FGlobalTabmanager::Get()->GetRootWindow();
-	if (RootWindow.IsValid())
-	{
-		FSlateApplication::Get().AddWindowAsNativeChild(PickAssetPathWindow.ToSharedRef(), RootWindow.ToSharedRef());
-	}
-	else
-	{
-		FSlateApplication::Get().AddWindow(PickAssetPathWindow.ToSharedRef());
-	}
-
-}
-
-void FFractureEditorModeToolkit::OnGenerateAssetPathChosen(const FString& InAssetPath, TArray<AActor*> Actors)
-{
-		UGeometryCollectionComponent* GeometryCollectionComponent = nullptr;//  = Cast<UGeometryCollectionComponent>(FractureContext.OriginalPrimitiveComponent);
-
-		if (Actors.Num() > 0)
-		{
-			AActor* FirstActor = Actors[0];
-
-			AGeometryCollectionActor* GeometryCollectionActor = Cast<AGeometryCollectionActor>(FirstActor);
-			GeometryCollectionActor = ConvertStaticMeshToGeometryCollection(InAssetPath, Actors);
-
-			GeometryCollectionComponent = GeometryCollectionActor->GetGeometryCollectionComponent();
-
-			FScopedColorEdit EditBoneColor = GeometryCollectionComponent->EditBoneSelection();
-			EditBoneColor.SetShowBoneColors(true);
-
-			// Move GC actor to source actors position and remove source actor from scene
-			const FVector ActorLocation(FirstActor->GetActorLocation());
-			GeometryCollectionActor->SetActorLocation(ActorLocation);
-
-			// Clear selection of mesh actor used to make GC before selecting, will cause details pane to not display geometry collection details.
-			GEditor->SelectNone(true, true, false);
-			GEditor->SelectActor(GeometryCollectionActor, true, true);
-
-			EditBoneColor.SelectBones(GeometryCollection::ESelectionMode::AllGeometry);
-
-			SetOutlinerComponents({ GeometryCollectionComponent });
-			SetBoneSelection(GeometryCollectionComponent, EditBoneColor.GetSelectedBones(), true);
-
-			GeometryCollectionComponent->MarkRenderDynamicDataDirty();
-			GeometryCollectionComponent->MarkRenderStateDirty();
-
-			for (AActor* Actor : Actors)
-			{
-				Actor->Destroy();
-			}
-		}
-}
-
-
-AActor* FFractureEditorModeToolkit::AddActor(ULevel* InLevel, UClass* Class)
-{
-	check(Class);
-
-	UWorld* World = InLevel->OwningWorld;
-	ULevel* DesiredLevel = InLevel;
-
-	// Transactionally add the actor.
-	AActor* Actor = NULL;
-	{
-		FScopedTransaction Transaction(NSLOCTEXT("UnrealEd", "AddActor", "Add Actor"));
-
-		FActorSpawnParameters SpawnInfo;
-		SpawnInfo.OverrideLevel = DesiredLevel;
-		SpawnInfo.SpawnCollisionHandlingOverride = ESpawnActorCollisionHandlingMethod::AlwaysSpawn;
-		SpawnInfo.ObjectFlags = RF_Transactional;
-		const auto Location = FVector(0);
-		const auto Rotation = FTransform(FVector(0)).GetRotation().Rotator();
-		Actor = World->SpawnActor(Class, &Location, &Rotation, SpawnInfo);
-
-		check(Actor);
-		Actor->InvalidateLightingCache();
-		Actor->PostEditMove(true);
-	}
-
-	// If this actor is part of any layers (set in its default properties), add them into the visible layers list.
-	ULayersSubsystem* Layers = GEditor->GetEditorSubsystem<ULayersSubsystem>();
-	Layers->SetLayersVisibility(Actor->Layers, true);
-
-	// Clean up.
-	Actor->MarkPackageDirty();
-	ULevel::LevelDirtiedEvent.Broadcast();
-
-	return Actor;
-}
-
-
-class AGeometryCollectionActor* FFractureEditorModeToolkit::CreateNewGeometryActor(const FString& InAssetPath, const FTransform& Transform, bool AddMaterials /*= false*/)
-{
-
-	FString UniquePackageName = InAssetPath;
-	FString UniqueAssetName = FPackageName::GetLongPackageAssetName(InAssetPath);
-
-	FAssetToolsModule& AssetToolsModule = FModuleManager::LoadModuleChecked<FAssetToolsModule>("AssetTools");
-	AssetToolsModule.Get().CreateUniqueAssetName(UniquePackageName, TEXT(""), UniquePackageName, UniqueAssetName);
-
-	UPackage* Package = CreatePackage(*UniquePackageName);
-	UGeometryCollection* InGeometryCollection = static_cast<UGeometryCollection*>(NewObject<UGeometryCollection>(Package, UGeometryCollection::StaticClass(), FName(*UniqueAssetName), RF_Transactional | RF_Public | RF_Standalone));
-
-	// Set default collision type here in the factory instead of the geometry collection constructor to only
-	// affect new geometry collections and not update the default for previously created objects.
-	if(InGeometryCollection)
-	{
-		InGeometryCollection->CollisionType = ECollisionTypeEnum::Chaos_Surface_Volumetric;
-		InGeometryCollection->ImplicitType = EImplicitTypeEnum::Chaos_Implicit_LevelSet;
-	}
-
-	// Create the new Geometry Collection actor
-	AGeometryCollectionActor* NewActor = Cast<AGeometryCollectionActor>(AddActor(GetSelectedLevel(), AGeometryCollectionActor::StaticClass()));
-	check(NewActor->GetGeometryCollectionComponent());
-
-	// Set the Geometry Collection asset in the new actor
-	NewActor->GetGeometryCollectionComponent()->SetRestCollection(InGeometryCollection);
-
-	// copy transform of original static mesh actor to this new actor
-	NewActor->SetActorLabel(UniqueAssetName);
-	NewActor->SetActorTransform(Transform);
-
-	// Mark relevant stuff dirty
-	FAssetRegistryModule::AssetCreated(InGeometryCollection);
-	InGeometryCollection->MarkPackageDirty();
-	Package->SetDirtyFlag(true);
-
-	return NewActor;
-}
-
-
-void FFractureEditorModeToolkit::ResetAsset()
-{
-	TSet<UGeometryCollectionComponent*> GeomCompSelection;
-	GetSelectedGeometryCollectionComponents(GeomCompSelection);
-	for (UGeometryCollectionComponent* GeometryCollectionComponent : GeomCompSelection)
-	{
-		FGeometryCollectionEdit GeometryCollectionEdit = GeometryCollectionComponent->EditRestCollection();
-		if (UGeometryCollection* GeometryCollectionObject = GeometryCollectionEdit.GetRestCollection())
-		{
-			
-			TSharedPtr<FGeometryCollection, ESPMode::ThreadSafe> GeometryCollectionPtr = GeometryCollectionObject->GetGeometryCollection();
-			if (FGeometryCollection* GeometryCollection = GeometryCollectionPtr.Get())
-			{
-				GeometryCollectionObject->Reset();			
-
-				// Rebuild Collection from recorded source assets.
-				for (const FGeometryCollectionSource& Source : GeometryCollectionObject->GeometrySource)
-				{
-					const UObject* SourceMesh = Source.SourceGeometryObject.TryLoad();
-					if (const UStaticMesh* SourceStaticMesh = Cast<UStaticMesh>(SourceMesh))
-					{
-						FGeometryCollectionConversion::AppendStaticMesh(SourceStaticMesh, Source.SourceMaterial, Source.LocalTransform, GeometryCollectionObject, true);
-					}
-					else if (const USkeletalMesh* SourceSkeletalMesh = Cast<USkeletalMesh>(SourceMesh))
-					{
-						// #todo (bmiller) Once we've settled on the right approach with static meshes, we'll need to apply the same strategy to skeletal mesh reconstruction.
-						// FGeometryCollectionConversion::AppendSkeletalMesh(SourceSkeletalMesh, Source.SourceMaterial, Source.LocalTransform, GeometryCollectionObject, true);
-					}
-				}
-
-				GeometryCollectionObject->InitializeMaterials();
-
-				FGeometryCollectionClusteringUtility::UpdateHierarchyLevelOfChildren(GeometryCollection, -1);
-				AddSingleRootNodeIfRequired(GeometryCollectionObject);
-				GeometryCollectionComponent->MarkRenderDynamicDataDirty();
-				GeometryCollectionComponent->MarkRenderStateDirty();
-			}
-			GeometryCollectionObject->MarkPackageDirty();
-		}
-	}
-	SetOutlinerComponents(GeomCompSelection.Array());
-}
-
-
-void FFractureEditorModeToolkit::ExecuteFracture(FFractureContext& FractureContext)
-{
-	if (ActiveTool != nullptr)
-	{
-		FractureContext.FracturedGeometryCollection->Modify();
-		ActiveTool->ExecuteFracture(FractureContext);
-		TSharedPtr<FGeometryCollection, ESPMode::ThreadSafe> GeometryCollectionPtr = FractureContext.FracturedGeometryCollection->GetGeometryCollection();
-		FGeometryCollection* OutGeometryCollection = GeometryCollectionPtr.Get();
-		FGeometryCollectionClusteringUtility::UpdateHierarchyLevelOfChildren(OutGeometryCollection, -1);
-
-		// Update Nanite resource data to correctly reflect modified geometry collection data
-		{
-			FractureContext.FracturedGeometryCollection->ReleaseResources();
-
-			if (FractureContext.FracturedGeometryCollection->EnableNanite)
-			{
-				FractureContext.FracturedGeometryCollection->NaniteData = UGeometryCollection::CreateNaniteData(OutGeometryCollection);
-			}
-			else
-			{
-				FractureContext.FracturedGeometryCollection->NaniteData = MakeUnique<FGeometryCollectionNaniteData>();
-			}
-
-			FractureContext.FracturedGeometryCollection->InitResources();
-		}
-
-		UGeometryCollectionComponent* GeometryCollectionComponent = Cast<UGeometryCollectionComponent>(FractureContext.OriginalPrimitiveComponent);
-
-		if (GeometryCollectionComponent != nullptr) // Create new GC actor from static mesh
-		{
-			FScopedColorEdit EditBoneColor = GeometryCollectionComponent->EditBoneSelection();
-
-			EditBoneColor.SelectBones(GeometryCollection::ESelectionMode::None);
-			SetBoneSelection(GeometryCollectionComponent, EditBoneColor.GetSelectedBones(), true);
-		}
-
-		GeometryCollectionComponent->MarkRenderDynamicDataDirty();
-		GeometryCollectionComponent->MarkRenderStateDirty();
-	}
-}
-
-
 bool FFractureEditorModeToolkit::IsGeometryCollectionSelected()
 {
 	USelection* SelectedActors = GEditor->GetSelectedActors();
@@ -1644,26 +898,6 @@ bool FFractureEditorModeToolkit::IsGeometryCollectionSelected()
 	return false;
 }
 
-bool FFractureEditorModeToolkit::IsStaticMeshSelected()
-{
-	USelection* SelectedActors = GEditor->GetSelectedActors();
-	for (FSelectionIterator Iter(*SelectedActors); Iter; ++Iter)
-	{
-		AActor* Actor = Cast<AActor>(*Iter);
-		if (Actor)
-		{
-			TInlineComponentArray<UStaticMeshComponent*> StaticMeshComponents;
-			Actor->GetComponents<UStaticMeshComponent>(StaticMeshComponents, true);
-
-			if (StaticMeshComponents.Num() > 0)
-			{
-				return true;
-			}
-		}
-	}
-	return false;
-}
-
 bool FFractureEditorModeToolkit::IsSelectedActorsInEditorWorld()
 {
 	USelection* SelectedActors = GEditor->GetSelectedActors();
@@ -1824,55 +1058,6 @@ void FFractureEditorModeToolkit::UpdateExplodedVectors(UGeometryCollectionCompon
 	}
 }
 
-AGeometryCollectionActor*  FFractureEditorModeToolkit::ConvertStaticMeshToGeometryCollection(const FString& InAssetPath, TArray<AActor*>& Actors)
-{
-	ensure(Actors.Num() > 0);
-	AActor* FirstActor = Actors[0];
-	const FString& Name = FirstActor->GetActorLabel();
-	const FVector FirstActorLocation(FirstActor->GetActorLocation());
-
-
-	AGeometryCollectionActor* NewActor = CreateNewGeometryActor(InAssetPath, FTransform(), true);
-
-	FGeometryCollectionEdit GeometryCollectionEdit = NewActor->GetGeometryCollectionComponent()->EditRestCollection(GeometryCollection::EEditUpdate::RestPhysicsDynamic);
-	UGeometryCollection* FracturedGeometryCollection = GeometryCollectionEdit.GetRestCollection();
-
-	for (AActor* Actor : Actors)
-	{
-		const FTransform ActorTransform(Actor->GetTransform());
-		const FVector ActorOffset(Actor->GetActorLocation() - FirstActor->GetActorLocation());
-
-		check(FracturedGeometryCollection);
-
-		TArray<UStaticMeshComponent*> StaticMeshComponents;
-		Actor->GetComponents<UStaticMeshComponent>(StaticMeshComponents, true);
-		for (int32 ii = 0, ni = StaticMeshComponents.Num(); ii < ni; ++ii)
-		{
-			// We're partial to static mesh components, here
-			UStaticMeshComponent* StaticMeshComponent = StaticMeshComponents[ii];
-			if (StaticMeshComponent != nullptr)
-			{
-				UStaticMesh* ComponentStaticMesh = StaticMeshComponent->GetStaticMesh();
-				FTransform ComponentTransform(StaticMeshComponent->GetComponentTransform());
-				ComponentTransform.SetTranslation((ComponentTransform.GetTranslation() - ActorTransform.GetTranslation()) + ActorOffset);
-				
-				// Record the contributing source on the asset.
-				FSoftObjectPath SourceSoftObjectPath(ComponentStaticMesh);			
-				TArray<UMaterialInterface*> SourceMaterials = StaticMeshComponent->GetMaterials();
-				FracturedGeometryCollection->GeometrySource.Add({ SourceSoftObjectPath, ComponentTransform, SourceMaterials });
-
-				FGeometryCollectionConversion::AppendStaticMesh(ComponentStaticMesh, SourceMaterials, ComponentTransform, FracturedGeometryCollection, true);
-			}
-		}
-
-		FracturedGeometryCollection->InitializeMaterials();
-	}
-	AddSingleRootNodeIfRequired(FracturedGeometryCollection);
-
-	return NewActor;
-}
-
-
 void FFractureEditorModeToolkit::OnOutlinerBoneSelectionChanged(UGeometryCollectionComponent* RootComponent, TArray<int32>& SelectedBones)
 {
 	FScopedTransaction Transaction(FractureTransactionContexts::SelectBoneContext, LOCTEXT("SelectGeometryCollectionBoneTransaction", "Select Bone"), RootComponent);
@@ -1895,7 +1080,6 @@ void FFractureEditorModeToolkit::OnOutlinerBoneSelectionChanged(UGeometryCollect
 	RootComponent->MarkRenderDynamicDataDirty();
 }
 
-
 FText FFractureEditorModeToolkit::GetStatisticsSummary() const
 {
 
diff --git a/Engine/Plugins/Experimental/ChaosEditor/Source/FractureEditor/Private/FractureToolAutoCluster.cpp b/Engine/Plugins/Experimental/ChaosEditor/Source/FractureEditor/Private/FractureToolAutoCluster.cpp
--- a/Engine/Plugins/Experimental/ChaosEditor/Source/FractureEditor/Private/FractureToolAutoCluster.cpp
+++ b/Engine/Plugins/Experimental/ChaosEditor/Source/FractureEditor/Private/FractureToolAutoCluster.cpp
@@ -15,25 +15,26 @@
 UFractureToolAutoCluster::UFractureToolAutoCluster(const FObjectInitializer& ObjInit)
   : Super(ObjInit) 
 {
-	Settings = NewObject<UFractureAutoClusterSettings>(GetTransientPackage(), UFractureAutoClusterSettings::StaticClass());
+	AutoClusterSettings = NewObject<UFractureAutoClusterSettings>(GetTransientPackage(), UFractureAutoClusterSettings::StaticClass());
+	AutoClusterSettings->OwnerTool = this;
 }
 
 
 // UFractureTool Interface
 FText UFractureToolAutoCluster::GetDisplayText() const
 {
-	return FText(NSLOCTEXT("Fracture", "FractureToolAutoCluster", "Auto")); 
+	return FText(NSLOCTEXT("FractureAutoCluster", "FractureToolAutoCluster", "Auto")); 
 }
 
 
 FText UFractureToolAutoCluster::GetTooltipText() const
 {
-	return FText(NSLOCTEXT("Fracture", "FractureToolAutoClusterToolTip", "Automatically group together pieces of a fractured mesh (based on your settings) and assign them within the Geometry Collection.")); 
+	return FText(NSLOCTEXT("FractureAutoCluster", "FractureToolAutoClusterToolTip", "Automatically group together pieces of a fractured mesh (based on your settings) and assign them within the Geometry Collection.")); 
 }
 
 FText UFractureToolAutoCluster::GetApplyText() const 
 { 
-	return FText(NSLOCTEXT("Fracture", "ExecuteAutoCluster", "Auto Cluster")); 
+	return FText(NSLOCTEXT("FractureAutoCluster", "ExecuteAutoCluster", "Auto Cluster")); 
 }
 
 FSlateIcon UFractureToolAutoCluster::GetToolIcon() const
@@ -45,7 +46,7 @@ FSlateIcon UFractureToolAutoCluster::GetToolIcon() const
 TArray<UObject*> UFractureToolAutoCluster::GetSettingsObjects() const 
 {
 	TArray<UObject*> AllSettings; 
-	AllSettings.Add(Settings);
+	AllSettings.Add(AutoClusterSettings);
 	return AllSettings;
 }
 
@@ -55,18 +56,27 @@ void UFractureToolAutoCluster::RegisterUICommand( FFractureEditorCommands* Bindi
 	BindingContext->AutoCluster = UICommandInfo;
 }
 
-void UFractureToolAutoCluster::ExecuteFracture(const FFractureContext& FractureContext)
+void UFractureToolAutoCluster::Execute(TWeakPtr<FFractureEditorModeToolkit> InToolkit)
 {
-	if (UGeometryCollectionComponent* GeometryCollectionComponent = Cast<UGeometryCollectionComponent>(FractureContext.OriginalPrimitiveComponent))
+	if (InToolkit.IsValid())
 	{
-		FScopedColorEdit EditBoneColor = GeometryCollectionComponent->EditBoneSelection();
-		int32 CurrentLevelView = EditBoneColor.GetViewLevel();
+		TSet<UGeometryCollectionComponent*> GeomCompSelection;
+		GetSelectedGeometryCollectionComponents(GeomCompSelection);
+		for (UGeometryCollectionComponent* GeometryCollectionComponent : GeomCompSelection)
+		{
+			FScopedColorEdit EditBoneColor = GeometryCollectionComponent->EditBoneSelection();
+			int32 CurrentLevelView = EditBoneColor.GetViewLevel();
 
-		UAutoClusterFractureCommand::ClusterChildBonesOfASingleMesh(GeometryCollectionComponent, Settings->AutoClusterMode, Settings->SiteCount);
+			UAutoClusterFractureCommand::ClusterChildBonesOfASingleMesh(GeometryCollectionComponent, AutoClusterSettings->AutoClusterMode, AutoClusterSettings->SiteCount);
 
-		EditBoneColor.ResetBoneSelection();
-		EditBoneColor.SetLevelViewMode(0);
-		EditBoneColor.SetLevelViewMode(CurrentLevelView);
+			EditBoneColor.ResetBoneSelection();
+			EditBoneColor.SetLevelViewMode(CurrentLevelView);
+
+			GeometryCollectionComponent->MarkRenderDynamicDataDirty();
+			GeometryCollectionComponent->MarkRenderStateDirty();
+		}
+
+		InToolkit.Pin()->SetOutlinerComponents(GeomCompSelection.Array());
 	}
 }
 
diff --git a/Engine/Plugins/Experimental/ChaosEditor/Source/FractureEditor/Private/FractureToolAutoCluster.h b/Engine/Plugins/Experimental/ChaosEditor/Source/FractureEditor/Private/FractureToolAutoCluster.h
--- a/Engine/Plugins/Experimental/ChaosEditor/Source/FractureEditor/Private/FractureToolAutoCluster.h
+++ b/Engine/Plugins/Experimental/ChaosEditor/Source/FractureEditor/Private/FractureToolAutoCluster.h
@@ -8,16 +8,15 @@
 #include "FractureToolAutoCluster.generated.h"
 
 
-UCLASS()
-class UFractureAutoClusterSettings
-	: public UObject
+UCLASS(DisplayName = "Auto Cluster", Category = "FractureTools")
+class UFractureAutoClusterSettings : public UFractureToolSettings
 {
-	GENERATED_BODY()
-
 public:
+	GENERATED_BODY()
 
-	UFractureAutoClusterSettings()
-		: AutoClusterMode(EFractureAutoClusterMode::BoundingBox)
+	UFractureAutoClusterSettings(const FObjectInitializer& ObjInit)
+		: Super(ObjInit)
+		, AutoClusterMode(EFractureAutoClusterMode::BoundingBox)
 		, SiteCount(10)
 	{}
 
@@ -26,12 +25,11 @@ public:
 
 	UPROPERTY(EditAnywhere, Category = AutoCluster, meta = (DisplayName = "Cluster Sites"))// , UIMin = "1", UIMax = "5000", ClampMin = "1"))
 	uint32 SiteCount;
-
 };
 
 
 UCLASS(DisplayName="AutoCluster", Category="FractureTools")
-class UFractureToolAutoCluster: public UFractureTool
+class UFractureToolAutoCluster: public UFractureModalTool
 {
 public:
 	GENERATED_BODY()
@@ -44,12 +42,10 @@ public:
 	virtual FText GetApplyText() const override; 
 	virtual FSlateIcon GetToolIcon() const override;
 	virtual TArray<UObject*> GetSettingsObjects() const override;
-
 	virtual void RegisterUICommand( FFractureEditorCommands* BindingContext );
 
-	virtual void ExecuteFracture(const FFractureContext& Context) override;
+	virtual void Execute(TWeakPtr<FFractureEditorModeToolkit> InToolkit) override;
 
 	UPROPERTY(EditAnywhere, Category = AutoCluster)
-	UFractureAutoClusterSettings* Settings;
-
+	UFractureAutoClusterSettings* AutoClusterSettings;
 };
\ No newline at end of file
diff --git a/Engine/Plugins/Experimental/ChaosEditor/Source/FractureEditor/Private/FractureToolBitmap.cpp b/Engine/Plugins/Experimental/ChaosEditor/Source/FractureEditor/Private/FractureToolBitmap.cpp
deleted file mode 100644
diff --git a/Engine/Plugins/Experimental/ChaosEditor/Source/FractureEditor/Private/FractureToolBitmap.h b/Engine/Plugins/Experimental/ChaosEditor/Source/FractureEditor/Private/FractureToolBitmap.h
deleted file mode 100644
diff --git a/Engine/Plugins/Experimental/ChaosEditor/Source/FractureEditor/Private/FractureToolBrick.cpp b/Engine/Plugins/Experimental/ChaosEditor/Source/FractureEditor/Private/FractureToolBrick.cpp
--- a/Engine/Plugins/Experimental/ChaosEditor/Source/FractureEditor/Private/FractureToolBrick.cpp
+++ b/Engine/Plugins/Experimental/ChaosEditor/Source/FractureEditor/Private/FractureToolBrick.cpp
@@ -14,24 +14,6 @@
 
 #define LOCTEXT_NAMESPACE "FractureBrick"
 
-void UFractureBrickSettings::PostEditChangeProperty(FPropertyChangedEvent& PropertyChangedEvent)
-{
-	if (OwnerTool != nullptr)
-	{
-		OwnerTool->PostEditChangeProperty(PropertyChangedEvent);
-	}
-	Super::PostEditChangeProperty(PropertyChangedEvent);
-}
-
-void UFractureBrickSettings::PostEditChangeChainProperty(struct FPropertyChangedChainEvent& PropertyChangedEvent)
-{
-	if (OwnerTool != nullptr)
-	{
-		OwnerTool->PostEditChangeChainProperty(PropertyChangedEvent);
-	}
-	Super::PostEditChangeChainProperty(PropertyChangedEvent);
-}
-
 
 UFractureToolBrick::UFractureToolBrick(const FObjectInitializer& ObjInit) 
 	: Super(ObjInit) 
@@ -64,34 +46,32 @@ void UFractureToolBrick::RegisterUICommand( FFractureEditorCommands* BindingCont
 TArray<UObject*> UFractureToolBrick::GetSettingsObjects() const 
 { 
 	TArray<UObject*> Settings; 
-	Settings.Add(GetMutableDefault<UFractureCommonSettings>());
-	Settings.Add(GetMutableDefault<UFractureBrickSettings>());
+	Settings.Add(CutterSettings);
+	Settings.Add(BrickSettings);
 	return Settings;
 }
 
 void UFractureToolBrick::GenerateBrickTransforms(const FBox& Bounds)
 {
-	const UFractureBrickSettings* LocalBrickSettings = GetMutableDefault<UFractureBrickSettings>();
-
 	const FVector& Min = Bounds.Min;
 	const FVector& Max = Bounds.Max;
 	const FVector Extents(Bounds.Max - Bounds.Min);
 
 	const FQuat HeaderRotation(FVector::UpVector, 1.5708);
 
-	const float HalfHeight = LocalBrickSettings->BrickHeight * 0.5f;
-	const float HalfDepth = LocalBrickSettings->BrickDepth * 0.5f;
-	const float HalfLength = LocalBrickSettings->BrickLength * 0.5f;
+	const float HalfHeight = BrickSettings->BrickHeight * 0.5f;
+	const float HalfDepth = BrickSettings->BrickDepth * 0.5f;
+	const float HalfLength = BrickSettings->BrickLength * 0.5f;
 
-	if (LocalBrickSettings->Bond == EFractureBrickBond::Stretcher)
+	if (BrickSettings->Bond == EFractureBrickBond::Stretcher)
 	{
 		bool OddY = false;
-		for (float yy = 0.f; yy <= Extents.Y; yy += LocalBrickSettings->BrickDepth)
+		for (float yy = 0.f; yy <= Extents.Y; yy += BrickSettings->BrickDepth)
 		{
 			bool Oddline = false;
-			for (float zz = HalfHeight; zz <= Extents.Z; zz += LocalBrickSettings->BrickHeight)
+			for (float zz = HalfHeight; zz <= Extents.Z; zz += BrickSettings->BrickHeight)
 			{
-				for (float xx = 0.f; xx <= Extents.X; xx += LocalBrickSettings->BrickLength)
+				for (float xx = 0.f; xx <= Extents.X; xx += BrickSettings->BrickLength)
 				{
 					FVector BrickPosition(Min + FVector(Oddline ^ OddY ? xx : xx + HalfLength, yy, zz));
 					BrickTransforms.Emplace(FTransform(BrickPosition));
@@ -101,14 +81,14 @@ void UFractureToolBrick::GenerateBrickTransforms(const FBox& Bounds)
 				OddY = !OddY;
 		}
 	}
-	else if (LocalBrickSettings->Bond == EFractureBrickBond::Stack)
+	else if (BrickSettings->Bond == EFractureBrickBond::Stack)
 	{
 		bool OddY = false;
-		for (float yy = 0.f; yy <= Extents.Y; yy += LocalBrickSettings->BrickDepth)
+		for (float yy = 0.f; yy <= Extents.Y; yy += BrickSettings->BrickDepth)
 		{
-			for (float zz = HalfHeight; zz <= Extents.Z; zz += LocalBrickSettings->BrickHeight)
+			for (float zz = HalfHeight; zz <= Extents.Z; zz += BrickSettings->BrickHeight)
 			{
-				for (float xx = 0.f; xx <= Extents.X; xx += LocalBrickSettings->BrickLength)
+				for (float xx = 0.f; xx <= Extents.X; xx += BrickSettings->BrickLength)
 				{
 					FVector BrickPosition(Min + FVector(OddY ? xx : xx + HalfLength, yy, zz));
 					BrickTransforms.Emplace(FTransform(BrickPosition));
@@ -117,18 +97,18 @@ void UFractureToolBrick::GenerateBrickTransforms(const FBox& Bounds)
 			OddY = !OddY;
 		}
 	}
-	else if (LocalBrickSettings->Bond == EFractureBrickBond::English)
+	else if (BrickSettings->Bond == EFractureBrickBond::English)
 	{
 		float HalfLengthDepthDifference = HalfLength - HalfDepth - HalfDepth;
 		bool OddY = false;
-		for (float yy = 0.f; yy <= Extents.Y; yy += LocalBrickSettings->BrickDepth)
+		for (float yy = 0.f; yy <= Extents.Y; yy += BrickSettings->BrickDepth)
 		{
 			bool Oddline = false;
-			for (float zz = HalfHeight; zz <= Extents.Z; zz += LocalBrickSettings->BrickHeight)
+			for (float zz = HalfHeight; zz <= Extents.Z; zz += BrickSettings->BrickHeight)
 			{
 				if (Oddline && !OddY) // header row
 				{
-					for (float xx = 0.f; xx <= Extents.X; xx += LocalBrickSettings->BrickDepth)
+					for (float xx = 0.f; xx <= Extents.X; xx += BrickSettings->BrickDepth)
 					{
 						FVector BrickPosition(Min + FVector(Oddline ^ OddY ? xx : xx + HalfDepth, yy + HalfDepth, zz));
 						BrickTransforms.Emplace(FTransform(HeaderRotation, BrickPosition));
@@ -136,7 +116,7 @@ void UFractureToolBrick::GenerateBrickTransforms(const FBox& Bounds)
 				}
 				else if(!Oddline) // stretchers
 				{
-					for (float xx = 0.f; xx <= Extents.X; xx += LocalBrickSettings->BrickLength)
+					for (float xx = 0.f; xx <= Extents.X; xx += BrickSettings->BrickLength)
 					{
 						FVector BrickPosition(Min + FVector(Oddline ^ OddY ? xx : xx + HalfLength, OddY ? yy + HalfLengthDepthDifference : yy - HalfLengthDepthDifference , zz));
 						BrickTransforms.Emplace(FTransform(BrickPosition));
@@ -147,15 +127,15 @@ void UFractureToolBrick::GenerateBrickTransforms(const FBox& Bounds)
 			OddY = !OddY;
 		}
 	}
-	else if (LocalBrickSettings->Bond == EFractureBrickBond::Header)
+	else if (BrickSettings->Bond == EFractureBrickBond::Header)
 	{
 		bool OddY = false;
-		for (float yy = 0.f; yy <= Extents.Y; yy += LocalBrickSettings->BrickLength)
+		for (float yy = 0.f; yy <= Extents.Y; yy += BrickSettings->BrickLength)
 		{
 			bool Oddline = false;
-			for (float zz = HalfHeight; zz <= Extents.Z; zz += LocalBrickSettings->BrickHeight)
+			for (float zz = HalfHeight; zz <= Extents.Z; zz += BrickSettings->BrickHeight)
 			{
-				for (float xx = 0.f; xx <= Extents.X; xx += LocalBrickSettings->BrickDepth)
+				for (float xx = 0.f; xx <= Extents.X; xx += BrickSettings->BrickDepth)
 				{
 					FVector BrickPosition(Min + FVector(Oddline ^ OddY ? xx : xx + HalfDepth, yy, zz));
 					BrickTransforms.Emplace(FTransform(HeaderRotation, BrickPosition));
@@ -165,20 +145,19 @@ void UFractureToolBrick::GenerateBrickTransforms(const FBox& Bounds)
 			OddY = !OddY;
 		}
 	}
-	else if (LocalBrickSettings->Bond == EFractureBrickBond::Flemish)
+	else if (BrickSettings->Bond == EFractureBrickBond::Flemish)
 	{
-		float HalfLengthDepthDifference = HalfLength - LocalBrickSettings->BrickDepth  ;
+		float HalfLengthDepthDifference = HalfLength - BrickSettings->BrickDepth  ;
 		bool OddY = false;
 		int32 RowY = 0;
-		for (float yy = 0.f; yy <= Extents.Y; yy += LocalBrickSettings->BrickDepth)
+		for (float yy = 0.f; yy <= Extents.Y; yy += BrickSettings->BrickDepth)
 		{
 			bool OddZ = false;
-			for (float zz = HalfHeight; zz <= Extents.Z; zz += LocalBrickSettings->BrickHeight)
+			for (float zz = HalfHeight; zz <= Extents.Z; zz += BrickSettings->BrickHeight)
 			{
 				bool OddX = OddZ;
 				for (float xx = 0.f; xx <= Extents.X; xx += HalfLength + HalfDepth)
 				{
-//					FVector BrickPosition(Min + FVector(OddY ? xx : xx + HalfLength, yy, zz));
 					FVector BrickPosition(Min + FVector(xx,yy,zz));
 					if (OddX)
 					{
@@ -202,30 +181,21 @@ void UFractureToolBrick::GenerateBrickTransforms(const FBox& Bounds)
 			}
 			OddY = !OddY;
 			++RowY;
-// 			if (RowY % 2)
-// 			{
-// 				yy += HalfLengthDepthDifference;
-// 			}
 		}
 	}
 
 	FVector BrickMax(HalfLength, HalfDepth, HalfHeight);
 	FVector BrickMin(-BrickMax);
 
-
 	for (const auto& Transform : BrickTransforms)
 	{
 		AddBoxEdges(Transform.TransformPosition(BrickMin), Transform.TransformPosition(BrickMax));
 	}
 }
 
-#if WITH_EDITOR
 void UFractureToolBrick::PostEditChangeChainProperty(struct FPropertyChangedChainEvent& PropertyChangedEvent)
 {
-	const UFractureCommonSettings* LocalCommonSettings = GetDefault<UFractureCommonSettings>();
-	const UFractureBrickSettings* LocalBrickSettings = GetMutableDefault<UFractureBrickSettings>();
-
-	FFractureContext FractureContext;
+	FFractureToolContext FractureContext;
 
 	FractureContext.Bounds = FBox(ForceInitToZero);
 
@@ -235,9 +205,9 @@ void UFractureToolBrick::PostEditChangeChainProperty(struct FPropertyChangedChai
 	SelectedActors.Reserve(SelectionSet->Num());
 
 	FractureContext.RandomSeed = FMath::Rand();
-	if (LocalCommonSettings->RandomSeed > -1)
+	if (CutterSettings->RandomSeed > -1)
 	{
-		FractureContext.RandomSeed = LocalCommonSettings->RandomSeed;
+		FractureContext.RandomSeed = CutterSettings->RandomSeed;
 	}
 
 	FBox SelectedMeshBounds(ForceInit);
@@ -254,12 +224,11 @@ void UFractureToolBrick::PostEditChangeChainProperty(struct FPropertyChangedChai
 		{
 			FractureContext.OriginalActor = Actor;
 			FractureContext.OriginalPrimitiveComponent = PrimitiveComponent;
-			// 				FractureContext.Transform = PrimitiveComponent->GetComponentTransform();
 			FractureContext.Transform = Actor->GetTransform();
 			FVector Origin;
 			FVector BoxExtent;
 			Actor->GetActorBounds(false, Origin, BoxExtent);
-			if (LocalCommonSettings->bGroupFracture)
+			if (CutterSettings->bGroupFracture)
 			{
 				FractureContext.Bounds += FBox::BuildAABB(Origin, BoxExtent);
 			}
@@ -271,20 +240,15 @@ void UFractureToolBrick::PostEditChangeChainProperty(struct FPropertyChangedChai
 		}
 	}
 
-	if (LocalCommonSettings->bGroupFracture)
+	if (CutterSettings->bGroupFracture)
 	{
 		GenerateBrickTransforms(FractureContext.Bounds);
 	}
 
 }
-#endif
-
 
 void UFractureToolBrick::Render(const FSceneView* View, FViewport* Viewport, FPrimitiveDrawInterface* PDI)
 {
-	const UFractureCommonSettings* LocalCommonSettings = GetDefault<UFractureCommonSettings>();
-
-
 	//TODO :Plane cutting drawing
 
 	for (const FTransform& Transform : BrickTransforms)
@@ -292,7 +256,7 @@ void UFractureToolBrick::Render(const FSceneView* View, FViewport* Viewport, FPr
 		PDI->DrawPoint(Transform.GetLocation(), FLinearColor::Green, 4.f, SDPG_Foreground);
 	}
 
-	if (LocalCommonSettings->bDrawDiagram)
+	if (CutterSettings->bDrawDiagram)
 	{
 		PDI->AddReserveLines(SDPG_Foreground, Edges.Num(), false, false);
 		for (int32 ii = 0, ni = Edges.Num(); ii < ni; ++ii)
@@ -322,7 +286,7 @@ void UFractureToolBrick::AddBoxEdges(const FVector& Min, const FVector& Max)
 	Edges.Emplace(MakeTuple(FVector(Min.X, Max.Y, Max.Z), Max));
 }
 
-void UFractureToolBrick::ExecuteFracture(const FFractureContext& FractureContext)
+void UFractureToolBrick::ExecuteFracture(const FFractureToolContext& FractureContext)
 {
 	if (FractureContext.FracturedGeometryCollection != nullptr)
 	{
@@ -332,14 +296,11 @@ void UFractureToolBrick::ExecuteFracture(const FFractureContext& FractureContext
 			BrickTransforms.Empty();
 			GenerateBrickTransforms(FractureContext.Bounds);
 
-			const UFractureCommonSettings* LocalCommonSettings = GetDefault<UFractureCommonSettings>();
-			const UFractureBrickSettings* LocalBrickSettings = GetMutableDefault<UFractureBrickSettings>();
-
 			const FQuat HeaderRotation(FVector::UpVector, 1.5708);
 
-			const float HalfHeight = LocalBrickSettings->BrickHeight * 0.5f;
-			const float HalfDepth = LocalBrickSettings->BrickDepth * 0.5f;
-			const float HalfLength = LocalBrickSettings->BrickLength * 0.5f;
+			const float HalfHeight = BrickSettings->BrickHeight * 0.5f;
+			const float HalfDepth = BrickSettings->BrickDepth * 0.5f;
+			const float HalfLength = BrickSettings->BrickLength * 0.5f;
 
 			FVector Max(HalfLength, HalfDepth, HalfHeight);
 			TArray<FBox> BricksToCut;
@@ -352,12 +313,12 @@ void UFractureToolBrick::ExecuteFracture(const FFractureContext& FractureContext
  			FPlanarCells VoronoiPlanarCells = FPlanarCells(BricksToCut);
 
 			FNoiseSettings NoiseSettings;
-			if (LocalCommonSettings->Amplitude > 0.0f)
+			if (CutterSettings->Amplitude > 0.0f)
 			{
-				NoiseSettings.Amplitude = LocalCommonSettings->Amplitude;
-				NoiseSettings.Frequency = LocalCommonSettings->Frequency;
-				NoiseSettings.Octaves = LocalCommonSettings->OctaveNumber;
-				NoiseSettings.PointSpacing = LocalCommonSettings->SurfaceResolution;
+				NoiseSettings.Amplitude = CutterSettings->Amplitude;
+				NoiseSettings.Frequency = CutterSettings->Frequency;
+				NoiseSettings.Octaves = CutterSettings->OctaveNumber;
+				NoiseSettings.PointSpacing = CutterSettings->SurfaceResolution;
 				VoronoiPlanarCells.InternalSurfaceMaterials.NoiseSettings = NoiseSettings;
 			}
 
@@ -366,9 +327,4 @@ void UFractureToolBrick::ExecuteFracture(const FFractureContext& FractureContext
 	}
 }
 
-bool UFractureToolBrick::CanExecuteFracture() const
-{
-	return FFractureEditorModeToolkit::IsLeafBoneSelected();
-}
-
 #undef LOCTEXT_NAMESPACE
diff --git a/Engine/Plugins/Experimental/ChaosEditor/Source/FractureEditor/Private/FractureToolBrick.h b/Engine/Plugins/Experimental/ChaosEditor/Source/FractureEditor/Private/FractureToolBrick.h
--- a/Engine/Plugins/Experimental/ChaosEditor/Source/FractureEditor/Private/FractureToolBrick.h
+++ b/Engine/Plugins/Experimental/ChaosEditor/Source/FractureEditor/Private/FractureToolBrick.h
@@ -3,6 +3,7 @@
 #pragma once
 
 #include "FractureTool.h"
+#include "FractureToolCutter.h"
 
 #include "FractureToolBrick.generated.h"
 
@@ -34,25 +35,20 @@ enum class EFractureBrickBond : uint8
 };
 
 
-UCLASS()
-class UFractureBrickSettings
-	: public UObject
+UCLASS(config = EditorPerProjectUserSettings)
+class UFractureBrickSettings : public UFractureToolSettings
 {
+public:
 	GENERATED_BODY()
 
-public:
-#if WITH_EDITOR
-	virtual void PostEditChangeProperty(FPropertyChangedEvent& PropertyChangedEvent) override;
-	virtual void PostEditChangeChainProperty(struct FPropertyChangedChainEvent& PropertyChangedEvent) override;
-#endif
-
-	UFractureBrickSettings() 
-	: Bond(EFractureBrickBond::Stretcher)
-	, Forward(EFractureBrickProjection::X)
-	, Up(EFractureBrickProjection::Z)
-	, BrickLength(194.f)
-	, BrickHeight(57.f)
-	, BrickDepth(92.f)
+	UFractureBrickSettings(const FObjectInitializer& ObjInit)
+		: Super(ObjInit)
+		, Bond(EFractureBrickBond::Stretcher)
+		, Forward(EFractureBrickProjection::X)
+		, Up(EFractureBrickProjection::Z)
+		, BrickLength(194.f)
+		, BrickHeight(57.f)
+		, BrickDepth(92.f)
 	{}
 
 	/** Forward Direction to project brick pattern. */
@@ -78,40 +74,29 @@ public:
 	/** Brick Height */
 	UPROPERTY(EditAnywhere, Category = Brick, meta = (UIMin = "0.1", UIMax = "500.0", ClampMin = "0.001"))
 	float BrickDepth;
-
-	UPROPERTY()
-	UFractureTool *OwnerTool;
-
 };
 
 
 UCLASS(DisplayName="Brick", Category="FractureTools")
-class UFractureToolBrick : public UFractureTool
+class UFractureToolBrick : public UFractureToolCutterBase
 {
 public:
 	GENERATED_BODY()
 
-	UFractureToolBrick(const FObjectInitializer& ObjInit);//  : Super(ObjInit) {}
+	UFractureToolBrick(const FObjectInitializer& ObjInit);
 
 	// UFractureTool Interface
 	virtual FText GetDisplayText() const override;
 	virtual FText GetTooltipText() const override;
 	virtual FSlateIcon GetToolIcon() const override;
-
-	virtual void RegisterUICommand( FFractureEditorCommands* BindingContext );
-	
-	virtual TArray<UObject*> GetSettingsObjects() const override;// { return TArray<UObject*>(); }
+	virtual void RegisterUICommand( FFractureEditorCommands* BindingContext ) override;
+	virtual TArray<UObject*> GetSettingsObjects() const override;
+	void PostEditChangeChainProperty(struct FPropertyChangedChainEvent& PropertyChangedEvent) override;
+	virtual void Render(const FSceneView* View, FViewport* Viewport, FPrimitiveDrawInterface* PDI) override;
+	virtual void ExecuteFracture(const FFractureToolContext& FractureContext) override;
 
 	void GenerateBrickTransforms(const FBox& Bounds);
-#if WITH_EDITOR
-	virtual void PostEditChangeChainProperty(struct FPropertyChangedChainEvent& PropertyChangedEvent) override;
-#endif
-
-
-	virtual void Render(const FSceneView* View, FViewport* Viewport, FPrimitiveDrawInterface* PDI) override;
 
-	virtual void ExecuteFracture(const FFractureContext& FractureContext) override;
-	virtual bool CanExecuteFracture() const override;
 
 private:
 	UPROPERTY()
diff --git a/Engine/Plugins/Experimental/ChaosEditor/Source/FractureEditor/Private/FractureToolCluster.cpp b/Engine/Plugins/Experimental/ChaosEditor/Source/FractureEditor/Private/FractureToolCluster.cpp
deleted file mode 100644
diff --git a/Engine/Plugins/Experimental/ChaosEditor/Source/FractureEditor/Private/FractureToolCluster.h b/Engine/Plugins/Experimental/ChaosEditor/Source/FractureEditor/Private/FractureToolCluster.h
deleted file mode 100644
diff --git a/Engine/Plugins/Experimental/ChaosEditor/Source/FractureEditor/Private/FractureToolClusterCutter.cpp b/Engine/Plugins/Experimental/ChaosEditor/Source/FractureEditor/Private/FractureToolClusterCutter.cpp
new file mode 100644
--- /dev/null
+++ b/Engine/Plugins/Experimental/ChaosEditor/Source/FractureEditor/Private/FractureToolClusterCutter.cpp
@@ -0,0 +1,77 @@
+// Copyright Epic Games, Inc. All Rights Reserved.
+
+
+#include "FractureToolClusterCutter.h"
+
+#include "FractureEditorStyle.h"
+
+#define LOCTEXT_NAMESPACE "FractureClustered"
+
+
+UFractureToolClusterCutter::UFractureToolClusterCutter(const FObjectInitializer& ObjInit)
+	: Super(ObjInit)
+{
+	ClusterSettings = NewObject<UFractureClusterCutterSettings>(GetTransientPackage(), UFractureClusterCutterSettings::StaticClass());
+	ClusterSettings->OwnerTool = this;
+}
+
+FText UFractureToolClusterCutter::GetDisplayText() const
+{
+	return FText(NSLOCTEXT("Fracture", "FractureToolCluster", "Cluster Voronoi Fracture"));
+}
+
+FText UFractureToolClusterCutter::GetTooltipText() const
+{
+	return FText(NSLOCTEXT("Fracture", "FractureToolClusterTooltip", "Cluster Voronoi Fracture creates additional points around a base Voronoi pattern, creating more variation.  Click the Fracture Button to commit the fracture to the geometry collection."));
+}
+
+FSlateIcon UFractureToolClusterCutter::GetToolIcon() const
+{
+	return FSlateIcon("FractureEditorStyle", "FractureEditor.Clustered");
+}
+
+void UFractureToolClusterCutter::RegisterUICommand( FFractureEditorCommands* BindingContext )
+{
+	UI_COMMAND_EXT( BindingContext, UICommandInfo, "Clustered", "Clustered", "Clustered Voronoi Fracture", EUserInterfaceActionType::ToggleButton, FInputChord() );
+	BindingContext->Clustered = UICommandInfo;
+}
+
+TArray<UObject*> UFractureToolClusterCutter::GetSettingsObjects() const
+{
+	TArray<UObject*> ReturnSettings;
+	ReturnSettings.Add(CutterSettings);
+	ReturnSettings.Add(ClusterSettings);
+	return ReturnSettings;
+}
+
+void UFractureToolClusterCutter::GenerateVoronoiSites(const FFractureToolContext &Context, TArray<FVector>& Sites)
+{
+ 	FRandomStream RandStream(Context.RandomSeed);
+	const int32 SiteCount = RandStream.RandRange(ClusterSettings->NumberClustersMin, ClusterSettings->NumberClustersMax);
+
+	const FVector Extent(Context.Bounds.Max - Context.Bounds.Min);
+
+	TArray<FVector> CenterSites;
+
+	Sites.Reserve(Sites.Num() + SiteCount);
+	for (int32 ii = 0; ii < SiteCount; ++ii)
+	{
+		CenterSites.Emplace(Context.Bounds.Min + FVector(RandStream.FRand(), RandStream.FRand(), RandStream.FRand()) * Extent);
+	}
+
+	for (int32 kk = 0, nk = CenterSites.Num(); kk < nk; ++kk)
+	{
+		const int32 SubSiteCount = RandStream.RandRange(ClusterSettings->SitesPerClusterMin, ClusterSettings->SitesPerClusterMax);
+
+		for (int32 ii = 0; ii < SubSiteCount; ++ii)
+		{
+			FVector V(RandStream.VRand());
+			V.Normalize();
+			V *= ClusterSettings->ClusterRadius + (RandStream.FRandRange(ClusterSettings->ClusterRadiusPercentageMin, ClusterSettings->ClusterRadiusPercentageMax) * Context.Bounds.GetExtent().GetAbsMax());
+			V += CenterSites[kk];
+			Sites.Emplace(V);
+		}
+	}
+}
+
+#undef LOCTEXT_NAMESPACE
\ No newline at end of file
diff --git a/Engine/Plugins/Experimental/ChaosEditor/Source/FractureEditor/Private/FractureToolClusterCutter.h b/Engine/Plugins/Experimental/ChaosEditor/Source/FractureEditor/Private/FractureToolClusterCutter.h
new file mode 100644
--- /dev/null
+++ b/Engine/Plugins/Experimental/ChaosEditor/Source/FractureEditor/Private/FractureToolClusterCutter.h
@@ -0,0 +1,80 @@
+// Copyright Epic Games, Inc. All Rights Reserved.
+
+#pragma once
+
+#include "FractureTool.h"
+#include "FractureToolCutter.h"
+
+#include "FractureToolClusterCutter.generated.h"
+
+
+UCLASS(config = EditorPerProjectUserSettings)
+class UFractureClusterCutterSettings : public UFractureToolSettings
+{
+public:
+	GENERATED_BODY()
+
+	UFractureClusterCutterSettings(const FObjectInitializer& ObjInit)
+		: Super(ObjInit)
+		, NumberClustersMin(8)
+		, NumberClustersMax(8)
+		, SitesPerClusterMin(2)
+		, SitesPerClusterMax(30)
+		, ClusterRadiusPercentageMin(0.1)
+		, ClusterRadiusPercentageMax(0.2)
+		, ClusterRadius(0.0f)
+	{}
+
+	/** Number of Clusters - Cluster Voronoi Method */
+	UPROPERTY(EditAnywhere, Category = ClusterVoronoi, meta = (DisplayName = "Minimum Cluster Sites", UIMin = "1", UIMax = "200", ClampMin = "1"))
+	int32 NumberClustersMin;
+
+	/** Number of Clusters - Cluster Voronoi Method */
+	UPROPERTY(EditAnywhere, Category = ClusterVoronoi, meta = (DisplayName = "Maximum Cluster Sites", UIMin = "1", UIMax = "200", ClampMin = "1"))
+	int32 NumberClustersMax;
+
+	/** Sites per # of Clusters - Cluster Voronoi Method */
+	UPROPERTY(EditAnywhere, Category = ClusterVoronoi, meta = (DisplayName = "Minimum Sites Per Cluster", UIMin = "0", UIMax = "200", ClampMin = "0"))
+	int32 SitesPerClusterMin;
+
+	UPROPERTY(EditAnywhere, Category = ClusterVoronoi, meta = (DisplayName = "Maximum Sites Per Cluster", UIMin = "0", UIMax = "200", ClampMin = "0"))
+	int32 SitesPerClusterMax;
+
+	/** Cluster's Radius - Cluster Voronoi Method */
+	UPROPERTY(EditAnywhere, Category = ClusterVoronoi, meta = (DisplayName = "Minimum distance from center as part of bounds max extent", UIMin = "0.0", UIMax = "1.0"))
+	float ClusterRadiusPercentageMin;
+
+	/** Cluster's Radius - Cluster Voronoi Method */
+	UPROPERTY(EditAnywhere, Category = ClusterVoronoi, meta = (DisplayName = "Maximum distance from center as part of bounds max extent", UIMin = "0.0", UIMax = "1.0"))
+	float ClusterRadiusPercentageMax;
+
+	/** Cluster's Radius - Cluster Voronoi Method */
+	UPROPERTY(EditAnywhere, Category = ClusterVoronoi)
+	float ClusterRadius;
+};
+
+
+UCLASS(DisplayName="Cluster Voronoi", Category="FractureTools")
+class UFractureToolClusterCutter : public UFractureToolVoronoiCutterBase
+{
+public:
+	GENERATED_BODY()
+
+	UFractureToolClusterCutter(const FObjectInitializer& ObjInit);
+
+	// UFractureTool Interface
+	virtual FText GetDisplayText() const override;
+	virtual FText GetTooltipText() const override;
+	virtual FSlateIcon GetToolIcon() const override;
+	virtual TArray<UObject*> GetSettingsObjects() const override;
+
+	virtual void RegisterUICommand( FFractureEditorCommands* BindingContext ) override;
+
+	
+	UPROPERTY(EditAnywhere, Category = Cluster)
+	UFractureClusterCutterSettings* ClusterSettings;
+
+protected:
+	void GenerateVoronoiSites(const FFractureToolContext &Context, TArray<FVector>& Sites) override;
+
+};
\ No newline at end of file
diff --git a/Engine/Plugins/Experimental/ChaosEditor/Source/FractureEditor/Private/FractureToolClusterMagnet.cpp b/Engine/Plugins/Experimental/ChaosEditor/Source/FractureEditor/Private/FractureToolClusterMagnet.cpp
new file mode 100644
--- /dev/null
+++ b/Engine/Plugins/Experimental/ChaosEditor/Source/FractureEditor/Private/FractureToolClusterMagnet.cpp
@@ -0,0 +1,373 @@
+// Copyright Epic Games, Inc. All Rights Reserved.
+
+#include "FractureToolClusterMagnet.h"
+
+#include "FractureTool.h"
+#include "FractureEditorStyle.h"
+#include "FractureEditorCommands.h"
+
+#include "Chaos/TriangleMesh.h"
+#include "PhysicsProxy/GeometryCollectionPhysicsProxy.h"
+#include "Chaos/MassProperties.h"
+#include "Chaos/Particles.h"
+#include "Chaos/Vector.h"
+
+#include "GeometryCollection/GeometryCollectionAlgo.h"
+#include "GeometryCollection/GeometryCollectionClusteringUtility.h"
+#include "GeometryCollection/GeometryCollectionProximityUtility.h"
+
+#include "GeometryCollection/GeometryCollectionComponent.h"
+
+
+#define LOCTEXT_NAMESPACE "FractureClusterMagnet"
+
+
+
+UFractureToolClusterMagnet::UFractureToolClusterMagnet(const FObjectInitializer& ObjInit)
+	: Super(ObjInit)
+{
+	ClusterMagnetSettings = NewObject<UFractureClusterMagnetSettings>(GetTransientPackage(), UFractureClusterMagnetSettings::StaticClass());
+	ClusterMagnetSettings->OwnerTool = this;
+}
+
+
+FText UFractureToolClusterMagnet::GetDisplayText() const
+{
+	return FText(NSLOCTEXT("Fracture", "FractureToolClusterMagnet", "Magnet"));
+}
+
+
+FText UFractureToolClusterMagnet::GetTooltipText() const
+{
+	return FText(NSLOCTEXT("Fracture", "FractureToolClusterMagnetToolTip", "Builds clusters at local level by collecting bones adjacent to clusters or bones with highest mass."));
+}
+
+FText UFractureToolClusterMagnet::GetApplyText() const
+{
+	return FText(NSLOCTEXT("Fracturet", "ExecuteClusterMagnet", "Cluster Magnet"));
+}
+
+FSlateIcon UFractureToolClusterMagnet::GetToolIcon() const
+{
+	return FSlateIcon("FractureEditorStyle", "FractureEditor.ClusterMagnet");
+}
+
+
+TArray<UObject*> UFractureToolClusterMagnet::GetSettingsObjects() const
+{
+	TArray<UObject*> Settings;
+	Settings.Add(ClusterMagnetSettings);
+	return Settings;
+}
+
+void UFractureToolClusterMagnet::RegisterUICommand(FFractureEditorCommands* BindingContext)
+{
+	UI_COMMAND_EXT(BindingContext, UICommandInfo, "ClusterMagnet", "Cluster Magnet", "Builds clusters at local level by collecting bones adjacent to clusters or bones with highest mass.", EUserInterfaceActionType::ToggleButton, FInputChord());
+	BindingContext->ClusterMagnet = UICommandInfo;
+}
+
+void UFractureToolClusterMagnet::Execute(TWeakPtr<FFractureEditorModeToolkit> InToolkit)
+{
+	
+	if (InToolkit.IsValid())
+	{
+		TSet<UGeometryCollectionComponent*> GeomCompSelection;
+		GetSelectedGeometryCollectionComponents(GeomCompSelection);
+		for (UGeometryCollectionComponent* GeometryCollectionComponent : GeomCompSelection)
+		{
+			const UGeometryCollection* RestCollection = GeometryCollectionComponent->GetRestCollection();
+			if (RestCollection)
+			{
+				FGeometryCollectionPtr GeometryCollection = RestCollection->GetGeometryCollection();
+
+				// We require certain attributes present to proceed.
+				if (!CheckPresenceOfNecessaryAttributes(GeometryCollection))
+				{
+					return;
+				}
+
+				// If no bones are selected, assume that we're working on the root's children
+				TArray<int32> SelectedBones = GeometryCollectionComponent->GetSelectedBones();
+				if (SelectedBones.Num() == 0)
+				{
+					FGeometryCollectionClusteringUtility::GetRootBones(GeometryCollection.Get(), SelectedBones);
+				}
+
+				for (int32 CurrentRoot : SelectedBones)
+				{
+					const TManagedArray<TSet<int32>>& Children = GeometryCollection->Children;
+					if (Children[CurrentRoot].Num() > 0)
+					{
+						const TManagedArray<int32>& Levels = GeometryCollection->GetAttribute<int32>("Level", FTransformCollection::TransformGroup);
+						const int32 OperatingLevel = Levels[CurrentRoot] + 1;
+						//const TArray<int32> TopNodes = GatherTopNodes(GeometryCollection, CurrentRoot);
+						const TSet<int32> TopNodes = Children[CurrentRoot];
+
+						UpdateMasses(GeometryCollection, TopNodes);
+
+						float CutoffMass = FindCutoffMass(ClusterMagnetSettings->MassPercentile, GeometryCollection, TopNodes);
+
+						// We have the connections for the leaf nodes of our geometry collection. We want to percolate those up to the top nodes.
+						TMap<int32, TSet<int32>> TopNodeConnectivity = InitializeConnectivity(TopNodes, GeometryCollection, OperatingLevel);
+
+						// Separate the top nodes into cluster magnets and a pool of available nodes.
+						TArray<FClusterMagnet> ClusterMagnets;
+						TSet<int32> RemainingPool;
+						SeparateClusterMagnets(GeometryCollection, TopNodes, CutoffMass, TopNodeConnectivity, ClusterMagnets, RemainingPool);
+
+						for (uint32 Iteration = 0; Iteration < ClusterMagnetSettings->Iterations; ++Iteration)
+						{
+							bool bNeighborsAbsorbed = false;
+
+							// each cluster gathers adjacent nodes from the pool
+							for (FClusterMagnet& ClusterMagnet : ClusterMagnets)
+							{
+								bNeighborsAbsorbed |= AbsorbClusterNeighbors(TopNodeConnectivity, ClusterMagnet, RemainingPool);
+							}
+
+							// early termination
+							if (!bNeighborsAbsorbed)
+							{
+								break;
+							}
+						}
+
+						// Create new clusters from the cluster magnets
+						for (const FClusterMagnet& ClusterMagnet : ClusterMagnets)
+						{
+							if (ClusterMagnet.ClusteredNodes.Num() > 1)
+							{
+								TArray<int32> NewChildren = ClusterMagnet.ClusteredNodes.Array();
+								NewChildren.Sort();
+								FGeometryCollectionClusteringUtility::ClusterBonesUnderNewNode(GeometryCollection.Get(), NewChildren[0], NewChildren, false, false);
+								
+							}
+						}
+
+						FGeometryCollectionClusteringUtility::UpdateHierarchyLevelOfChildren(GeometryCollection.Get(), CurrentRoot);
+					}
+				}
+
+				FScopedColorEdit EditBoneColor = GeometryCollectionComponent->EditBoneSelection();
+				EditBoneColor.ResetBoneSelection();
+				EditBoneColor.ResetHighlightedBones();
+				InToolkit.Pin()->SetBoneSelection(GeometryCollectionComponent, EditBoneColor.GetSelectedBones(), true);
+
+				GeometryCollectionComponent->MarkRenderDynamicDataDirty();
+				GeometryCollectionComponent->MarkRenderStateDirty();				
+			}
+		}
+
+		InToolkit.Pin()->SetOutlinerComponents(GeomCompSelection.Array());
+	}
+}
+
+
+bool UFractureToolClusterMagnet::CheckPresenceOfNecessaryAttributes(const FGeometryCollectionPtr GeometryCollection) const
+{
+	if (!GeometryCollection->HasAttribute("Level", FTransformCollection::TransformGroup))
+	{
+		UE_LOG(LogFractureTool, Error, TEXT("Cannot execute Cluster Magnet tool: missing Level attribute."));
+		return false;
+	}
+
+	if (!GeometryCollection->HasAttribute("Proximity", FGeometryCollection::GeometryGroup))
+	{
+		UE_LOG(LogFractureTool, Error, TEXT("Cannot execute Cluster Magnet tool: missing Proximity attribute."));
+		return false;
+	}
+
+	return true;
+}
+
+float UFractureToolClusterMagnet::FindCutoffMass(float Percentile, const FGeometryCollectionPtr GeometryCollection, const TSet<int32>& TopNodes) const
+{
+	check(Percentile >= 0.0);
+	check(Percentile <= 1.0);
+	
+	const TManagedArray<float>& Mass = GeometryCollection->GetAttribute<float>("Mass", FTransformCollection::TransformGroup);
+
+	// Collect the top node masses, sort by mass, return the threshold mass at the cutoff threshold. 
+	TArray<float> Masses;
+	Masses.Reserve(TopNodes.Num());
+	for (int32 Index : TopNodes)
+	{
+		Masses.Add(Mass[Index]);
+	}
+	Masses.Sort();
+
+	int32 ThresholdIndex = FMath::Floor(TopNodes.Num() * Percentile);
+	return Masses[ThresholdIndex];
+}
+
+
+void UFractureToolClusterMagnet::UpdateMasses(FGeometryCollectionPtr GeometryCollection, const TSet<int32>& TopNodes) const 
+{
+	if (!GeometryCollection->HasAttribute("Mass", FTransformCollection::TransformGroup))
+	{
+		GeometryCollection->AddAttribute<Chaos::FReal>("Mass", FTransformCollection::TransformGroup);
+		UE_LOG(LogFractureTool, Warning, TEXT("Added Mass attribute needed to execute ClusterMagnet."));
+	}
+	
+	TArray<FTransform> Transform;
+	GeometryCollectionAlgo::GlobalMatrices(GeometryCollection->Transform, GeometryCollection->Parent, Transform);
+
+	const TManagedArray<FVector>& Vertex = GeometryCollection->Vertex;
+	const TManagedArray<int32>& BoneMap = GeometryCollection->BoneMap;
+	
+	Chaos::TParticles<float, 3> MassSpaceParticles;
+	MassSpaceParticles.AddParticles(Vertex.Num());
+	for (int32 Idx = 0; Idx < Vertex.Num(); ++Idx)
+	{
+		MassSpaceParticles.X(Idx) = Transform[BoneMap[Idx]].TransformPosition(Vertex[Idx]);
+	}
+
+	for (int32 Index : TopNodes)
+	{
+		UpdateMasses(GeometryCollection, MassSpaceParticles, Index);
+	}
+}
+
+void UFractureToolClusterMagnet::UpdateMasses(FGeometryCollectionPtr GeometryCollection, const Chaos::TParticles<float, 3>& MassSpaceParticles, int32 TransformIndex) const
+{
+	const TManagedArray<TSet<int32>>& Children = GeometryCollection->Children;
+	const TManagedArray<bool>& Visible = GeometryCollection->Visible;
+	const TManagedArray<int32>& FaceCount = GeometryCollection->FaceCount;
+	const TManagedArray<int32>& FaceStart = GeometryCollection->FaceStart;
+	const TManagedArray<int32>& TransformToGeometryIndex = GeometryCollection->TransformToGeometryIndex;
+	const TManagedArray<FIntVector>& Indices = GeometryCollection->Indices;
+	TManagedArray<float>& Mass = GeometryCollection->GetAttribute<float>("Mass", FTransformCollection::TransformGroup);
+	
+	if (Children[TransformIndex].Num() == 0) // leaf node
+	{
+		int32 GeometryIndex = TransformToGeometryIndex[TransformIndex];
+
+		TUniquePtr<Chaos::TTriangleMesh<float>> TriMesh(
+			CreateTriangleMesh(
+				FaceStart[GeometryIndex],
+				FaceCount[GeometryIndex],
+				Visible,
+				Indices,
+				true));
+		
+		float Volume = 0.0;
+		Chaos::TVector<Chaos::FReal,3> CenterOfMass;
+		Chaos::CalculateVolumeAndCenterOfMass(MassSpaceParticles, TriMesh->GetElements(), Volume, CenterOfMass);
+
+		// Since we're only interested in relative mass, we assume density = 1.0
+		Mass[TransformIndex] = Volume;
+	}
+	else
+	{
+		// Recurse to children and sum the masses for this node
+		float LocalMass = 0.0;
+		for (int32 ChildIndex : Children[TransformIndex])
+		{
+			UpdateMasses(GeometryCollection, MassSpaceParticles, ChildIndex);
+			LocalMass += Mass[ChildIndex];
+		}
+
+		Mass[TransformIndex] = LocalMass;
+	}
+}
+
+TMap<int32, TSet<int32>> UFractureToolClusterMagnet::InitializeConnectivity(const TSet<int32>& TopNodes, FGeometryCollectionPtr GeometryCollection, int32 OperatingLevel) const
+{
+	FGeometryCollectionProximityUtility::UpdateProximity(GeometryCollection.Get());
+
+	TMap<int32, TSet<int32>> ConnectivityMap;
+	for (int32 Index : TopNodes)
+	{
+		// Collect the proximity indices of all the leaf nodes under this top node,
+		// traced back up to its parent top node, so that all connectivity describes
+		// relationships only between top nodes.
+		TSet<int32> Connections;
+		CollectTopNodeConnections(GeometryCollection, Index, OperatingLevel, Connections);
+		Connections.Remove(Index);
+
+		// Remove any connections outside the current operating branch.
+		ConnectivityMap.Add(Index, Connections.Intersect(TopNodes));
+	}
+
+	return ConnectivityMap;
+}
+
+void UFractureToolClusterMagnet::CollectTopNodeConnections(FGeometryCollectionPtr GeometryCollection, int32 Index, int32 OperatingLevel, TSet<int32>& OutConnections) const
+{
+	const TManagedArray<TSet<int32>>& Children = GeometryCollection->Children;
+	if (Children[Index].Num() == 0) // leaf node
+	{
+		const TManagedArray<TSet<int32>>& Proximity = GeometryCollection->GetAttribute<TSet<int32>>("Proximity", FGeometryCollection::GeometryGroup);
+		const TManagedArray<int32>& GeometryToTransformIndex = GeometryCollection->TransformIndex;
+		const TManagedArray<int32>& TransformToGeometryIndex = GeometryCollection->TransformToGeometryIndex;
+
+
+		for (int32 Neighbor : Proximity[TransformToGeometryIndex[Index]])
+		{
+			int32 NeighborTransformIndex = GeometryToTransformIndex[Neighbor];
+			OutConnections.Add(FGeometryCollectionClusteringUtility::GetParentOfBoneAtSpecifiedLevel(GeometryCollection.Get(), NeighborTransformIndex, OperatingLevel));
+		}		
+	}
+	else
+	{
+		for (int32 ChildIndex : Children[Index])
+		{
+			CollectTopNodeConnections(GeometryCollection, ChildIndex, OperatingLevel, OutConnections);
+		}
+	}	
+}
+
+void UFractureToolClusterMagnet::SeparateClusterMagnets(
+	const FGeometryCollectionPtr GeometryCollection,
+	const TSet<int32>& TopNodes,
+	float CutoffMass,
+	const TMap<int32, TSet<int32>>& TopNodeConnectivity,
+	TArray<FClusterMagnet>& OutClusterMagnets,
+	TSet<int32>& OutRemainingPool) const
+{
+	// Push any top nodes over the mass threshold into cluster magnets.
+	
+	const TManagedArray<float>& Mass = GeometryCollection->GetAttribute<float>("Mass", FTransformCollection::TransformGroup);
+	
+	OutClusterMagnets.Reserve(TopNodes.Num());
+	OutRemainingPool.Reserve(TopNodes.Num());
+
+	for (int32 Index : TopNodes)
+	{
+		if (Mass[Index] > CutoffMass)
+		{
+			OutClusterMagnets.AddDefaulted();
+			FClusterMagnet& NewMagnet = OutClusterMagnets.Last();
+			NewMagnet.ClusteredNodes.Add(Index);
+			NewMagnet.Connections = TopNodeConnectivity[Index];
+		}
+		else
+		{
+			OutRemainingPool.Add(Index);
+		}
+	}
+}
+
+bool UFractureToolClusterMagnet::AbsorbClusterNeighbors(const TMap<int32, TSet<int32>> TopNodeConnectivity, FClusterMagnet& OutClusterMagnet, TSet<int32>& OutRemainingPool) const
+{
+	// Return true if neighbors were absorbed.
+	bool bNeighborsAbsorbed = false;
+
+	TSet<int32> NewConnections;
+	for (int32 NeighborIndex : OutClusterMagnet.Connections)
+	{
+		// If the neighbor is still in the pool, absorb it and its connections.
+		if (OutRemainingPool.Contains(NeighborIndex))
+		{
+			OutClusterMagnet.ClusteredNodes.Add(NeighborIndex);
+			NewConnections.Append(TopNodeConnectivity[NeighborIndex]);
+			OutRemainingPool.Remove(NeighborIndex);
+			bNeighborsAbsorbed = true;
+		}
+	}
+	OutClusterMagnet.Connections.Append(NewConnections);
+
+	return bNeighborsAbsorbed;
+}
+
+#undef LOCTEXT_NAMESPACE
\ No newline at end of file
diff --git a/Engine/Plugins/Experimental/ChaosEditor/Source/FractureEditor/Private/FractureToolClusterMagnet.h b/Engine/Plugins/Experimental/ChaosEditor/Source/FractureEditor/Private/FractureToolClusterMagnet.h
new file mode 100644
--- /dev/null
+++ b/Engine/Plugins/Experimental/ChaosEditor/Source/FractureEditor/Private/FractureToolClusterMagnet.h
@@ -0,0 +1,78 @@
+// Copyright Epic Games, Inc. All Rights Reserved.
+
+#pragma once
+
+#include "FractureTool.h"
+
+#include "FractureToolClusterMagnet.generated.h"
+
+namespace Chaos
+{
+	template<class T, int d>
+	class TParticles;
+}
+
+
+struct FClusterMagnet
+{
+	TSet<int32> ClusteredNodes;
+	TSet<int32> Connections;
+};
+
+
+UCLASS(DisplayName = "Cluster Magnet", Category = "FractureTools")
+class UFractureClusterMagnetSettings : public UFractureToolSettings
+{
+public:
+	GENERATED_BODY()
+
+		UFractureClusterMagnetSettings(const FObjectInitializer& ObjInit)
+		: Super(ObjInit)
+		, Iterations(1)
+		, MassPercentile(0.9)
+	{}
+
+	UPROPERTY(EditAnywhere, Category = ClusterMagnet, meta = (DisplayName = "Iterations", UIMin = "1", UIMax = "100", ClampMin = "1", ClampMax = "500"))
+	uint32 Iterations;
+
+	UPROPERTY(EditAnywhere, Category = ClusterMagnet, meta = (DisplayName = "Mass percentile to form magnets", UIMin = "0.0", UIMax = "1.0", ClampMin = "0.0", ClampMax = "1.0"))
+	float MassPercentile;
+
+};
+
+
+UCLASS(DisplayName = "ClusterMagnet", Category = "FractureTools")
+class UFractureToolClusterMagnet : public UFractureModalTool
+{
+public:
+	GENERATED_BODY()
+
+	using FGeometryCollectionPtr = TSharedPtr<FGeometryCollection, ESPMode::ThreadSafe>;
+		
+	UFractureToolClusterMagnet(const FObjectInitializer& ObjInit);
+
+	// UFractureTool Interface
+	virtual FText GetDisplayText() const override;
+	virtual FText GetTooltipText() const override;
+	virtual FText GetApplyText() const override;
+	virtual FSlateIcon GetToolIcon() const override;
+	virtual TArray<UObject*> GetSettingsObjects() const override;
+	virtual void RegisterUICommand(FFractureEditorCommands* BindingContext);
+
+	virtual void Execute(TWeakPtr<FFractureEditorModeToolkit> InToolkit) override;
+
+	UPROPERTY(EditAnywhere, Category = AutoCluster)
+	UFractureClusterMagnetSettings* ClusterMagnetSettings;
+
+private:
+	bool CheckPresenceOfNecessaryAttributes(const FGeometryCollectionPtr GeometryCollection) const;
+	//TArray<int32> GatherTopNodes(const FGeometryCollectionPtr GeometryCollection) const;
+	float FindCutoffMass(float Percentile, const FGeometryCollectionPtr GeometryCollection, const TSet<int32>& TopNodes) const;
+	void UpdateMasses(FGeometryCollectionPtr GeometryCollection, const TSet<int32>& TopNodes) const;
+	void UpdateMasses(FGeometryCollectionPtr GeometryCollection, const Chaos::TParticles<float, 3>& MassSpaceParticles, int32 TransformIndex) const;
+	TMap<int32, TSet<int32>> InitializeConnectivity(const TSet<int32>& TopNodes, FGeometryCollectionPtr GeometryCollection, int32 OperatingLevel) const;
+	void CollectTopNodeConnections(FGeometryCollectionPtr GeometryCollection, int32 Index, int32 OperatingLevel, TSet<int32>& OutConnections) const;
+	void SeparateClusterMagnets(const FGeometryCollectionPtr GeometryCollection, const TSet<int32>& TopNodes, float CutoffMass, const TMap<int32, TSet<int32>>& TopNodeConnectivity, TArray<FClusterMagnet>& OutClusterMagnets, TSet<int32>& OutRemainingPool) const;
+	bool AbsorbClusterNeighbors(const TMap<int32, TSet<int32>> TopNodeConnectivity, FClusterMagnet& OutClusterMagnets, TSet<int32>& OutRemainingPool) const;
+
+};
diff --git a/Engine/Plugins/Experimental/ChaosEditor/Source/FractureEditor/Private/FractureToolClustering.cpp b/Engine/Plugins/Experimental/ChaosEditor/Source/FractureEditor/Private/FractureToolClustering.cpp
new file mode 100644
--- /dev/null
+++ b/Engine/Plugins/Experimental/ChaosEditor/Source/FractureEditor/Private/FractureToolClustering.cpp
@@ -0,0 +1,269 @@
+// Copyright Epic Games, Inc. All Rights Reserved.
+
+#include "FractureToolClustering.h"
+
+#include "GeometryCollection/GeometryCollectionComponent.h"
+#include "GeometryCollection/GeometryCollection.h"
+#include "GeometryCollection/GeometryCollectionClusteringUtility.h"
+
+
+#define LOCTEXT_NAMESPACE "FractureToolClusteringOps"
+
+FText UFractureToolFlattenAll::GetDisplayText() const
+{
+	return FText(NSLOCTEXT("FractureToolClusteringOps", "FractureToolFlattenAll", "Flatten"));
+}
+
+FText UFractureToolFlattenAll::GetTooltipText() const
+{
+	return FText(NSLOCTEXT("FractureToolClusteringOps", "FractureToolFlattenAllTooltip", "Flattens all bones to level 1"));
+}
+
+FSlateIcon UFractureToolFlattenAll::GetToolIcon() const
+{
+	return FSlateIcon("FractureEditorStyle", "FractureEditor.Flatten");
+}
+
+void UFractureToolFlattenAll::RegisterUICommand(FFractureEditorCommands* BindingContext)
+{
+	UI_COMMAND_EXT(BindingContext, UICommandInfo, "Flatten", "Flatten", "Flattens all bones to level 1.", EUserInterfaceActionType::Button, FInputChord());
+	BindingContext->Flatten = UICommandInfo;
+}
+
+void UFractureToolFlattenAll::Execute(TWeakPtr<FFractureEditorModeToolkit> InToolkit)
+{
+	if (InToolkit.IsValid())
+	{
+		TSet<UGeometryCollectionComponent*> GeomCompSelection;
+		GetSelectedGeometryCollectionComponents(GeomCompSelection);
+		for (UGeometryCollectionComponent* GeometryCollectionComponent : GeomCompSelection)
+		{
+			FGeometryCollectionEdit GCEdit = GeometryCollectionComponent->EditRestCollection();
+			if (UGeometryCollection* GCObject = GCEdit.GetRestCollection())
+			{
+				TSharedPtr<FGeometryCollection, ESPMode::ThreadSafe> GeometryCollectionPtr = GCObject->GetGeometryCollection();
+				if (FGeometryCollection* GeometryCollection = GeometryCollectionPtr.Get())
+				{
+					AddAdditionalAttributesIfRequired(GCObject);
+
+					AddSingleRootNodeIfRequired(GCObject);
+
+					int32 NumElements = GCObject->NumElements(FGeometryCollection::TransformGroup);
+					TArray<int32> Elements;
+					Elements.Reserve(NumElements);
+
+					for (int32 Element = 0; Element < NumElements; ++Element)
+					{
+						if (GeometryCollection->Parent[Element] != FGeometryCollection::Invalid)
+						{
+							Elements.Add(Element);
+						}
+					}
+
+					if (Elements.Num() > 0)
+					{
+						FGeometryCollectionClusteringUtility::ClusterBonesUnderExistingRoot(GeometryCollection, Elements);
+					}
+
+					FGeometryCollectionClusteringUtility::UpdateHierarchyLevelOfChildren(GeometryCollection, -1);
+
+					FScopedColorEdit EditBoneColor = GeometryCollectionComponent->EditBoneSelection();
+					EditBoneColor.ResetBoneSelection();
+
+					InToolkit.Pin()->OnSetLevelViewValue(1);
+
+					GeometryCollectionComponent->MarkRenderDynamicDataDirty();
+					GeometryCollectionComponent->MarkRenderStateDirty();
+				}
+			}
+		}
+
+		InToolkit.Pin()->SetOutlinerComponents(GeomCompSelection.Array());
+	}	
+}
+
+
+
+FText UFractureToolCluster::GetDisplayText() const
+{
+	return FText(NSLOCTEXT("FractureToolClusteringOps", "FractureToolCluster", "Cluster"));
+}
+
+FText UFractureToolCluster::GetTooltipText() const
+{
+	return FText(NSLOCTEXT("FractureToolClusteringOps", "FractureToolClusterTooltip", "Clusters selected bones under a new parent."));
+}
+
+FSlateIcon UFractureToolCluster::GetToolIcon() const
+{
+	return FSlateIcon("FractureEditorStyle", "FractureEditor.Cluster");
+}
+
+void UFractureToolCluster::RegisterUICommand(FFractureEditorCommands* BindingContext)
+{
+	UI_COMMAND_EXT(BindingContext, UICommandInfo, "Cluster", "Cluster", "Clusters selected bones under a new parent.", EUserInterfaceActionType::Button, FInputChord());
+	BindingContext->Cluster = UICommandInfo;
+}
+
+void UFractureToolCluster::Execute(TWeakPtr<FFractureEditorModeToolkit> InToolkit)
+{
+	if (InToolkit.IsValid())
+	{
+		TSet<UGeometryCollectionComponent*> GeomCompSelection;
+		GetSelectedGeometryCollectionComponents(GeomCompSelection);
+		for (UGeometryCollectionComponent* GeometryCollectionComponent : GeomCompSelection)
+		{
+			const TArray<int32> SelectedBones = GeometryCollectionComponent->GetSelectedBones();
+
+			if (SelectedBones.Num() > 1)
+			{
+				FGeometryCollectionEdit GCEdit = GeometryCollectionComponent->EditRestCollection();
+				if (UGeometryCollection* GCObject = GCEdit.GetRestCollection())
+				{
+					TSharedPtr<FGeometryCollection, ESPMode::ThreadSafe> GeometryCollectionPtr = GCObject->GetGeometryCollection();
+					if (FGeometryCollection* GeometryCollection = GeometryCollectionPtr.Get())
+					{
+						const TManagedArray<TSet<int32>>& Children = GeometryCollection->GetAttribute<TSet<int32>>("Children", FGeometryCollection::TransformGroup);
+
+						// sort the selection list so ClusterBonesUnderNewNode() happens in the correct order for leaf nodes
+						TArray<int32> SortedSelectedBones;
+						SortedSelectedBones.Reserve(SelectedBones.Num());
+						for (int32 SelectedBone : SelectedBones)
+						{
+							if (Children[SelectedBone].Num() > 0)
+							{
+								SortedSelectedBones.Insert(SelectedBone, 0);
+							}
+							else
+							{
+								SortedSelectedBones.Add(SelectedBone);
+							}
+						}
+						// cluster Selected Bones under the first selected bone
+						int32 InsertAtIndex = SortedSelectedBones[0];
+
+						FGeometryCollectionClusteringUtility::ClusterBonesUnderNewNode(GeometryCollection, InsertAtIndex, SortedSelectedBones, false);
+						FGeometryCollectionClusteringUtility::UpdateHierarchyLevelOfChildren(GeometryCollection, -1);
+
+						FScopedColorEdit EditBoneColor = GeometryCollectionComponent->EditBoneSelection();
+						EditBoneColor.ResetBoneSelection();
+						EditBoneColor.ResetHighlightedBones();
+						GeometryCollectionComponent->MarkRenderDynamicDataDirty();
+						GeometryCollectionComponent->MarkRenderStateDirty();
+						InToolkit.Pin()->SetBoneSelection(GeometryCollectionComponent, EditBoneColor.GetSelectedBones(), true);
+					}
+				}
+			}
+		}
+
+		InToolkit.Pin()->SetOutlinerComponents(GeomCompSelection.Array());
+	}
+}
+
+
+
+FText UFractureToolUncluster::GetDisplayText() const
+{
+	return FText(NSLOCTEXT("FractureToolClusteringOps", "FractureToolUncluster", "Uncluster"));
+}
+
+FText UFractureToolUncluster::GetTooltipText() const
+{
+	return FText(NSLOCTEXT("FractureToolClusteringOps", "FractureToolUnclusterTooltip", "Remove parent cluster and move bones up a level."));
+}
+
+FSlateIcon UFractureToolUncluster::GetToolIcon() const
+{
+	return FSlateIcon("FractureEditorStyle", "FractureEditor.Uncluster");
+}
+
+void UFractureToolUncluster::RegisterUICommand(FFractureEditorCommands* BindingContext)
+{
+	UI_COMMAND_EXT(BindingContext, UICommandInfo, "Uncluster", "Uncluster", "Remove parent cluster and move bones up a level.", EUserInterfaceActionType::Button, FInputChord());
+	BindingContext->Uncluster = UICommandInfo;
+}
+
+void UFractureToolUncluster::Execute(TWeakPtr<FFractureEditorModeToolkit> InToolkit)
+{
+	if (InToolkit.IsValid())
+	{
+		int32 FractureLevel = InToolkit.Pin()->GetLevelViewValue();
+		TSet<UGeometryCollectionComponent*> GeomCompSelection;
+		GetSelectedGeometryCollectionComponents(GeomCompSelection);
+		for (UGeometryCollectionComponent* GeometryCollectionComponent : GeomCompSelection)
+		{
+			// scoped edit of collection
+			FGeometryCollectionEdit GeometryCollectionEdit = GeometryCollectionComponent->EditRestCollection();
+			if (UGeometryCollection* GeometryCollectionObject = GeometryCollectionEdit.GetRestCollection())
+			{
+				TSharedPtr<FGeometryCollection, ESPMode::ThreadSafe> GeometryCollectionPtr = GeometryCollectionObject->GetGeometryCollection();
+				if (FGeometryCollection* GeometryCollection = GeometryCollectionPtr.Get())
+				{
+					FGeometryCollectionClusteringUtility::UpdateHierarchyLevelOfChildren(GeometryCollection, -1);
+					FGeometryCollectionClusteringUtility::CollapseSelectedHierarchy(FractureLevel, GeometryCollectionComponent->GetSelectedBones(), GeometryCollection);
+
+					FScopedColorEdit EditBoneColor = GeometryCollectionComponent->EditBoneSelection();
+					EditBoneColor.ResetBoneSelection();
+					EditBoneColor.ResetHighlightedBones();
+					GeometryCollectionComponent->MarkRenderDynamicDataDirty();
+					GeometryCollectionComponent->MarkRenderStateDirty();
+					InToolkit.Pin()->SetBoneSelection(GeometryCollectionComponent, EditBoneColor.GetSelectedBones(), true);
+				}
+			}
+		}
+		InToolkit.Pin()->SetOutlinerComponents(GeomCompSelection.Array());
+	}
+}
+
+
+
+
+FText UFractureToolMoveUp::GetDisplayText() const
+{
+	return FText(NSLOCTEXT("FractureToolClusteringOps", "FractureToolMoveUp", "Level Up"));
+}
+
+FText UFractureToolMoveUp::GetTooltipText() const
+{
+	return FText(NSLOCTEXT("FractureToolClusteringOps", "FractureToolMoveUpTooltip", "Move bones up a level."));
+}
+
+FSlateIcon UFractureToolMoveUp::GetToolIcon() const
+{
+	return FSlateIcon("FractureEditorStyle", "FractureEditor.MoveUp");
+}
+
+void UFractureToolMoveUp::RegisterUICommand(FFractureEditorCommands* BindingContext)
+{
+	UI_COMMAND_EXT(BindingContext, UICommandInfo, "MoveUp", "Level Up", "Move bones up a level.", EUserInterfaceActionType::Button, FInputChord());
+	BindingContext->MoveUp = UICommandInfo;
+}
+
+void UFractureToolMoveUp::Execute(TWeakPtr<FFractureEditorModeToolkit> InToolkit)
+{
+	if (InToolkit.IsValid())
+	{
+		TSet<UGeometryCollectionComponent*> GeomCompSelection;
+		GetSelectedGeometryCollectionComponents(GeomCompSelection);
+		for (UGeometryCollectionComponent* GeometryCollectionComponent : GeomCompSelection)
+		{
+			FGeometryCollectionEdit GeometryCollectionEdit = GeometryCollectionComponent->EditRestCollection();
+			if (UGeometryCollection* GeometryCollectionObject = GeometryCollectionEdit.GetRestCollection())
+			{
+				TSharedPtr<FGeometryCollection, ESPMode::ThreadSafe> GeometryCollectionPtr = GeometryCollectionObject->GetGeometryCollection();
+				if (FGeometryCollection* GeometryCollection = GeometryCollectionPtr.Get())
+				{
+					TArray<int32> Selected = GeometryCollectionComponent->GetSelectedBones();
+					FGeometryCollectionClusteringUtility::MoveUpOneHierarchyLevel(GeometryCollection, Selected);
+
+					GeometryCollectionComponent->MarkRenderDynamicDataDirty();
+					GeometryCollectionComponent->MarkRenderStateDirty();
+				}
+			}
+		}
+
+		InToolkit.Pin()->SetOutlinerComponents(GeomCompSelection.Array());
+	}
+}
+
+#undef LOCTEXT_NAMESPACE
diff --git a/Engine/Plugins/Experimental/ChaosEditor/Source/FractureEditor/Private/FractureToolClustering.h b/Engine/Plugins/Experimental/ChaosEditor/Source/FractureEditor/Private/FractureToolClustering.h
new file mode 100644
--- /dev/null
+++ b/Engine/Plugins/Experimental/ChaosEditor/Source/FractureEditor/Private/FractureToolClustering.h
@@ -0,0 +1,75 @@
+// Copyright Epic Games, Inc. All Rights Reserved.
+
+#pragma once
+
+#include "FractureTool.h"
+
+#include "FractureToolClustering.generated.h"
+
+
+UCLASS(DisplayName = "Flatten All", Category = "FractureTools")
+class UFractureToolFlattenAll : public UFractureActionTool
+{
+public:
+	GENERATED_BODY()
+
+	UFractureToolFlattenAll(const FObjectInitializer& ObjInit) : Super(ObjInit) {}
+
+	// UFractureActionTool Interface
+	virtual FText GetDisplayText() const override;
+	virtual FText GetTooltipText() const override;
+	virtual FSlateIcon GetToolIcon() const override;
+	virtual void RegisterUICommand(FFractureEditorCommands* BindingContext) override;
+	virtual void Execute(TWeakPtr<FFractureEditorModeToolkit> InToolkit) override;
+};
+
+
+UCLASS(DisplayName = "Cluster", Category = "FractureTools")
+class UFractureToolCluster : public UFractureActionTool
+{
+public:
+	GENERATED_BODY()
+
+	UFractureToolCluster(const FObjectInitializer& ObjInit) : Super(ObjInit) {}
+
+	// UFractureActionTool Interface
+	virtual FText GetDisplayText() const override;
+	virtual FText GetTooltipText() const override;
+	virtual FSlateIcon GetToolIcon() const override;
+	virtual void RegisterUICommand(FFractureEditorCommands* BindingContext) override;
+	virtual void Execute(TWeakPtr<FFractureEditorModeToolkit> InToolkit) override;
+};
+
+
+UCLASS(DisplayName = "Uncluster", Category = "FractureTools")
+class UFractureToolUncluster : public UFractureActionTool
+{
+public:
+	GENERATED_BODY()
+
+	UFractureToolUncluster(const FObjectInitializer& ObjInit) : Super(ObjInit) {}
+
+	// UFractureActionTool Interface
+	virtual FText GetDisplayText() const override;
+	virtual FText GetTooltipText() const override;
+	virtual FSlateIcon GetToolIcon() const override;
+	virtual void RegisterUICommand(FFractureEditorCommands* BindingContext) override;
+	virtual void Execute(TWeakPtr<FFractureEditorModeToolkit> InToolkit) override;
+};
+
+
+UCLASS(DisplayName = "MoveUp", Category = "FractureTools")
+class UFractureToolMoveUp : public UFractureActionTool
+{
+public:
+	GENERATED_BODY()
+
+	UFractureToolMoveUp(const FObjectInitializer& ObjInit) : Super(ObjInit) {}
+
+	// UFractureActionTool Interface
+	virtual FText GetDisplayText() const override;
+	virtual FText GetTooltipText() const override;
+	virtual FSlateIcon GetToolIcon() const override;
+	virtual void RegisterUICommand(FFractureEditorCommands* BindingContext) override;
+	virtual void Execute(TWeakPtr<FFractureEditorModeToolkit> InToolkit) override;
+};
diff --git a/Engine/Plugins/Experimental/ChaosEditor/Source/FractureEditor/Private/FractureToolCutter.cpp b/Engine/Plugins/Experimental/ChaosEditor/Source/FractureEditor/Private/FractureToolCutter.cpp
new file mode 100644
--- /dev/null
+++ b/Engine/Plugins/Experimental/ChaosEditor/Source/FractureEditor/Private/FractureToolCutter.cpp
@@ -0,0 +1,242 @@
+
+#include "FractureToolCutter.h"
+
+#include "GeometryCollection/GeometryCollectionObject.h"
+#include "GeometryCollection/GeometryCollectionComponent.h"
+#include "GeometryCollection/GeometryCollectionAlgo.h"
+
+#include "Voronoi/Voronoi.h"
+#include "PlanarCut.h"
+
+
+UFractureToolCutterBase::UFractureToolCutterBase(const FObjectInitializer& ObjInit)
+	: Super(ObjInit)
+{
+	CutterSettings = NewObject<UFractureCutterSettings>(GetTransientPackage(), UFractureCutterSettings::StaticClass());
+	CutterSettings->OwnerTool = this;
+}
+
+bool UFractureToolCutterBase::CanExecute() const
+{
+	if (!IsGeometryCollectionSelected())
+	{
+		return false;
+	}
+
+	if (IsStaticMeshSelected())
+	{
+		return false;
+	}
+
+	return IsLeafBoneSelected();
+}
+
+
+void UFractureToolCutterBase::GetFractureContexts(TArray<FFractureToolContext>& FractureContexts) const
+{
+	FRandomStream RandomStream(CutterSettings->RandomSeed > -1 ? CutterSettings->RandomSeed : FMath::Rand());
+
+	TSet<UGeometryCollectionComponent*> GeometryCollectionComponents;
+	GetSelectedGeometryCollectionComponents(GeometryCollectionComponents);
+
+	for (UGeometryCollectionComponent* GeometryCollectionComponent : GeometryCollectionComponents)
+	{
+		FGeometryCollectionEdit RestCollection = GeometryCollectionComponent->EditRestCollection(GeometryCollection::EEditUpdate::None);
+		UGeometryCollection* FracturedGeometryCollection = RestCollection.GetRestCollection();
+		if (FracturedGeometryCollection == nullptr)
+		{
+			continue;
+		}
+
+		TSharedPtr<FGeometryCollection, ESPMode::ThreadSafe> GeometryCollectionPtr = FracturedGeometryCollection->GetGeometryCollection();
+		FGeometryCollection* OutGeometryCollection = GeometryCollectionPtr.Get();
+
+		TArray<int32> FilteredBones = FilterBones(GeometryCollectionComponent->GetSelectedBones(), OutGeometryCollection);
+
+
+		const TManagedArray<FTransform>& Transform = OutGeometryCollection->GetAttribute<FTransform>("Transform", FGeometryCollection::TransformGroup);
+		const TManagedArray<int32>& TransformToGeometryIndex = OutGeometryCollection->GetAttribute<int32>("TransformToGeometryIndex", FGeometryCollection::TransformGroup);
+		const TManagedArray<FBox>& BoundingBoxes = OutGeometryCollection->GetAttribute<FBox>("BoundingBox", FGeometryCollection::GeometryGroup);
+
+		TArray<FTransform> Transforms;
+		GeometryCollectionAlgo::GlobalMatrices(Transform, OutGeometryCollection->Parent, Transforms);
+
+
+		TMap<int32, FBox> BoundsToBone;
+		for (int32 Idx = 0, ni = FracturedGeometryCollection->NumElements(FGeometryCollection::TransformGroup); Idx < ni; ++Idx)
+		{
+			if (TransformToGeometryIndex[Idx] > -1)
+			{
+				ensure(TransformToGeometryIndex[Idx] > -1);
+				BoundsToBone.Add(Idx, BoundingBoxes[TransformToGeometryIndex[Idx]].TransformBy(Transforms[Idx]));
+			}
+		}
+
+		if (CutterSettings->bGroupFracture)
+		{
+			GenerateFractureToolContext(
+				GeometryCollectionComponent->GetOwner(),
+				GeometryCollectionComponent,
+				FracturedGeometryCollection,
+				FilteredBones,
+				BoundsToBone,
+				TransformToGeometryIndex,
+				CutterSettings->RandomSeed,
+				FractureContexts);
+		}
+		else
+		{
+			for (int32 BoneIndex : FilteredBones)
+			{
+				GenerateFractureToolContext(
+					GeometryCollectionComponent->GetOwner(),
+					GeometryCollectionComponent,
+					FracturedGeometryCollection,
+					{ BoneIndex },
+					BoundsToBone,
+					TransformToGeometryIndex,
+					CutterSettings->RandomSeed + FractureContexts.Num(),
+					FractureContexts);
+			}
+		}
+		
+	}
+}
+
+
+bool UFractureToolCutterBase::IsLeafBoneSelected()
+{
+	TSet<UGeometryCollectionComponent*> GeomCompSelection;
+	GetSelectedGeometryCollectionComponents(GeomCompSelection);
+	for (UGeometryCollectionComponent* GeometryCollectionComponent : GeomCompSelection)
+	{
+		const TArray<int32> SelectedBones = GeometryCollectionComponent->GetSelectedBones();
+
+		if (SelectedBones.Num() > 0)
+		{
+			if (const UGeometryCollection* GCObject = GeometryCollectionComponent->GetRestCollection())
+			{
+				TSharedPtr<FGeometryCollection, ESPMode::ThreadSafe> GeometryCollectionPtr = GCObject->GetGeometryCollection();
+				if (const FGeometryCollection* GeometryCollection = GeometryCollectionPtr.Get())
+				{
+					const TManagedArray<TSet<int32>>& Children = GeometryCollection->GetAttribute<TSet<int32>>("Children", FGeometryCollection::TransformGroup);
+
+					for (int32 BoneIndex : SelectedBones)
+					{
+						if (Children[BoneIndex].Num() == 0)
+						{
+							return true;
+						}
+					}
+				}
+			}
+		}
+	}
+	return false;
+}
+
+TArray<int32> UFractureToolCutterBase::FilterBones(const TArray<int32>& SelectedBonesOriginal, const FGeometryCollection* const GeometryCollection) const
+{
+	FRandomStream RandomStream(CutterSettings->RandomSeed > -1 ? CutterSettings->RandomSeed : FMath::Rand());
+
+	// Keep only leaf nodes and implement skip probability
+	const TManagedArray<TSet<int32>>& Children = GeometryCollection->GetAttribute<TSet<int32>>("Children", FGeometryCollection::TransformGroup);
+
+	TArray<int32> SelectedBones;
+	SelectedBones.Reserve(SelectedBonesOriginal.Num());
+	for (int32 BoneIndex : SelectedBonesOriginal)
+	{
+		if (Children[BoneIndex].Num() == 0 && RandomStream.FRand() < CutterSettings->ChanceToFracture)
+		{
+			SelectedBones.Add(BoneIndex);
+		}
+	}
+
+	return SelectedBones;
+}
+
+
+
+UFractureToolVoronoiCutterBase::UFractureToolVoronoiCutterBase(const FObjectInitializer& ObjInit)
+	: Super(ObjInit)
+{
+	for (int32 ii = 0; ii < 100; ++ii)
+	{
+		Colors.Emplace(FMath::FRand(), FMath::FRand(), FMath::FRand());;
+	}
+}
+
+void UFractureToolVoronoiCutterBase::Render(const FSceneView* View, FViewport* Viewport, FPrimitiveDrawInterface* PDI)
+{
+	if (CutterSettings->bDrawSites)
+	{
+		for (const FVector& Site : VoronoiSites)
+		{
+			PDI->DrawPoint(Site, FLinearColor::Green, 4.f, SDPG_Foreground);
+		}
+	}
+
+	if (CutterSettings->bDrawDiagram)
+	{
+		PDI->AddReserveLines(SDPG_Foreground, VoronoiEdges.Num(), false, false);
+		for (int32 ii = 0, ni = VoronoiEdges.Num(); ii < ni; ++ii)
+		{
+			PDI->DrawLine(VoronoiEdges[ii].Get<0>(), VoronoiEdges[ii].Get<1>(), Colors[CellMember[ii] % 100], SDPG_Foreground);
+		}
+	}
+}
+
+void UFractureToolVoronoiCutterBase::FractureContextChanged()
+{
+	TArray<FFractureToolContext> FractureContexts;
+	GetFractureContexts(FractureContexts);
+
+	VoronoiSites.Empty();
+	CellMember.Empty();
+	VoronoiEdges.Empty();;
+
+	for (FFractureToolContext& FractureContext : FractureContexts)
+	{
+		// Move the local bounds to the actor so we we'll draw in the correct location
+		FractureContext.Bounds = FractureContext.Bounds.TransformBy(FractureContext.Transform);
+		GenerateVoronoiSites(FractureContext, VoronoiSites);
+		if (CutterSettings->bDrawDiagram)
+		{
+			GetVoronoiEdges(VoronoiSites, FractureContext.Bounds, VoronoiEdges, CellMember);
+		}
+	}
+}
+
+void UFractureToolVoronoiCutterBase::ExecuteFracture(const FFractureToolContext& FractureContext)
+{
+	if (FractureContext.FracturedGeometryCollection != nullptr)
+	{
+		TSharedPtr<FGeometryCollection, ESPMode::ThreadSafe> GeometryCollectionPtr = FractureContext.FracturedGeometryCollection->GetGeometryCollection();
+		if (FGeometryCollection* GeometryCollection = GeometryCollectionPtr.Get())
+		{
+			TArray<FVector> Sites;
+			GenerateVoronoiSites(FractureContext, Sites);
+			FBox VoronoiBounds = FractureContext.Bounds;
+			if (CutterSettings->Amplitude > 0.0f)
+			{
+				// expand bounds to make sure noise-perturbed voronoi cells still contain the whole input mesh
+				VoronoiBounds = VoronoiBounds.ExpandBy(CutterSettings->Amplitude);
+			}
+			FVoronoiDiagram Voronoi(Sites, VoronoiBounds, .1f);
+
+			FPlanarCells VoronoiPlanarCells = FPlanarCells(Sites, Voronoi);
+
+			FNoiseSettings NoiseSettings;
+			if (CutterSettings->Amplitude > 0.0f)
+			{
+				NoiseSettings.Amplitude = CutterSettings->Amplitude;
+				NoiseSettings.Frequency = CutterSettings->Frequency;
+				NoiseSettings.Octaves = CutterSettings->OctaveNumber;
+				NoiseSettings.PointSpacing = CutterSettings->SurfaceResolution;
+				VoronoiPlanarCells.InternalSurfaceMaterials.NoiseSettings = NoiseSettings;
+			}
+
+			CutMultipleWithPlanarCells(VoronoiPlanarCells, *GeometryCollection, FractureContext.SelectedBones);
+		}
+	}
+}
diff --git a/Engine/Plugins/Experimental/ChaosEditor/Source/FractureEditor/Private/FractureToolCutter.h b/Engine/Plugins/Experimental/ChaosEditor/Source/FractureEditor/Private/FractureToolCutter.h
new file mode 100644
--- /dev/null
+++ b/Engine/Plugins/Experimental/ChaosEditor/Source/FractureEditor/Private/FractureToolCutter.h
@@ -0,0 +1,117 @@
+#pragma once
+
+#include "FractureTool.h"
+
+#include "FractureToolCutter.generated.h"
+
+
+/** Settings specifically related to the one-time destructive fracturing of a mesh **/
+UCLASS(config = EditorPerProjectUserSettings)
+class UFractureCutterSettings : public UFractureToolSettings
+{
+public:
+
+	GENERATED_BODY()
+
+	UFractureCutterSettings(const FObjectInitializer& ObjInit)
+		: Super(ObjInit)
+		, RandomSeed(-1)
+		, ChanceToFracture(1.0)
+		, bGroupFracture(true)
+		, bDrawSites(false)
+		, bDrawDiagram(true)
+		, Amplitude(0.0f)
+		, Frequency(0.1f)
+		, OctaveNumber(4)
+		, SurfaceResolution(10)
+	{}
+
+	/** Random number generator seed for repeatability */
+	UPROPERTY(EditAnywhere, Category = CommonFracture, meta = (DisplayName = "Random Seed", UIMin = "-1", UIMax = "1000", ClampMin = "-1"))
+	int32 RandomSeed;
+
+	/** Chance to shatter each mesh.  Useful when shattering multiple selected meshes.  */
+	UPROPERTY(EditAnywhere, Category = CommonFracture, meta = (DisplayName = "Chance To Fracture Per Mesh", UIMin = "0.0", UIMax = "1.0", ClampMin = "0.0", ClampMax = "1.0"))
+	float ChanceToFracture;
+
+	/** Generate a fracture pattern across all selected meshes.  */
+	UPROPERTY(EditAnywhere, Category = CommonFracture, meta = (DisplayName = "Group Fracture"))
+	bool bGroupFracture;
+
+	/** Generate a fracture pattern across all selected meshes.  */
+	UPROPERTY(EditAnywhere, Category = CommonFracture, meta = (DisplayName = "Draw Sites"))
+	bool bDrawSites;
+
+	/** Generate a fracture pattern across all selected meshes.  */
+	UPROPERTY(EditAnywhere, Category = CommonFracture, meta = (DisplayName = "Draw Diagram"))
+	bool bDrawDiagram;
+
+	/** Size of the noise displacement in centimeters */
+	UPROPERTY(EditAnywhere, Category = Noise, meta = (UIMin = "0.0"))
+	float Amplitude;
+
+	/** Period of the Perlin noise.  Smaller values will create noise faces that are smoother */
+	UPROPERTY(EditAnywhere, Category = Noise)
+	float Frequency;
+
+	/** Number of fractal layers of Perlin noise to apply.  Smaller values (1 or 2) will create noise that looks like gentle rolling hills, while larger values (> 4) will tend to look more like craggy mountains */
+	UPROPERTY(EditAnywhere, Category = Noise, meta = (UIMin = "1"))
+	int32 OctaveNumber;
+
+	/** Spacing between vertices on cut surfaces, where noise is added.  Larger spacing between vertices will create more efficient meshes with fewer triangles, but less resolution to see the shape of the added noise  */
+	UPROPERTY(EditAnywhere, Category = Noise, meta = (UIMin = "1", ClampMin = "0.1"))
+	float SurfaceResolution;
+
+};
+
+
+UCLASS(Abstract, DisplayName = "Cutter Base", Category = "FractureTools")
+class UFractureToolCutterBase : public UFractureInteractiveTool
+{
+public:
+	GENERATED_BODY()
+
+	UFractureToolCutterBase(const FObjectInitializer& ObjInit);
+
+	/** This is the Text that will appear on the button to execute the fracture **/
+	virtual FText GetApplyText() const override { return FText(NSLOCTEXT("Fracture", "ExecuteFracture", "Fracture")); }
+
+	virtual bool CanExecute() const override;
+
+	virtual void GetFractureContexts(TArray<FFractureToolContext>& FractureContexts) const;
+
+	TArray<int32> FilterBones(const TArray<int32>& SelectedBonesOriginal, const FGeometryCollection* const GeometryCollection) const override;
+
+protected:
+	UPROPERTY(EditAnywhere, Category = Slicing)
+	UFractureCutterSettings* CutterSettings;
+
+private:
+	static bool IsLeafBoneSelected();
+
+};
+
+
+UCLASS(Abstract, DisplayName = "Voronoi Base", Category = "FractureTools")
+class UFractureToolVoronoiCutterBase : public UFractureToolCutterBase
+{
+public:
+	GENERATED_BODY()
+
+	UFractureToolVoronoiCutterBase(const FObjectInitializer& ObjInit);
+
+	virtual void Render(const FSceneView* View, FViewport* Viewport, FPrimitiveDrawInterface* PDI) override;
+
+	virtual void FractureContextChanged() override;
+	virtual void ExecuteFracture(const FFractureToolContext& FractureContext) override;
+
+protected:
+	virtual void GenerateVoronoiSites(const FFractureToolContext& Context, TArray<FVector>& Sites) {}
+
+private:
+	TArray<int32> CellMember;
+	TArray<TTuple<FVector, FVector>> VoronoiEdges;
+	TArray<FVector> VoronoiSites;
+	TArray<FLinearColor> Colors;
+};
+
diff --git a/Engine/Plugins/Experimental/ChaosEditor/Source/FractureEditor/Private/FractureToolGenerators.cpp b/Engine/Plugins/Experimental/ChaosEditor/Source/FractureEditor/Private/FractureToolGenerators.cpp
new file mode 100644
--- /dev/null
+++ b/Engine/Plugins/Experimental/ChaosEditor/Source/FractureEditor/Private/FractureToolGenerators.cpp
@@ -0,0 +1,349 @@
+// Copyright Epic Games, Inc. All Rights Reserved.
+
+#include "FractureToolGenerators.h"
+
+#include "Editor.h"
+#include "ScopedTransaction.h"
+#include "Engine/Selection.h"
+#include "SCreateAssetFromObject.h"
+#include "AssetToolsModule.h"
+#include "AssetRegistryModule.h"
+#include "Layers/LayersSubsystem.h"
+#include "Engine/SkeletalMesh.h"
+
+#include "GeometryCollection/GeometryCollectionComponent.h"
+#include "GeometryCollection/GeometryCollectionActor.h"
+#include "GeometryCollection/GeometryCollection.h"
+#include "GeometryCollection/GeometryCollectionClusteringUtility.h"
+#include "GeometryCollection/GeometryCollectionConversion.h"
+
+
+#define LOCTEXT_NAMESPACE "FractureToolGenerators"
+
+
+FText UFractureToolGenerateAsset::GetDisplayText() const
+{
+	return FText(NSLOCTEXT("Fracture", "FractureToolGenerateAsset", "New"));
+}
+
+FText UFractureToolGenerateAsset::GetTooltipText() const
+{
+	return FText(NSLOCTEXT("Fracture", "FractureToolGenerateAssetTooltip", "Generate Geometry Collection Asset from static meshes contained in selected actors."));
+}
+
+FSlateIcon UFractureToolGenerateAsset::GetToolIcon() const
+{
+	return FSlateIcon("FractureEditorStyle", "FractureEditor.GenerateAsset");
+}
+
+void UFractureToolGenerateAsset::RegisterUICommand(FFractureEditorCommands* BindingContext)
+{
+	UI_COMMAND_EXT(BindingContext, UICommandInfo, "GenerateAsset", "New", "Generate Asset", EUserInterfaceActionType::ToggleButton, FInputChord());
+	BindingContext->GenerateAsset = UICommandInfo;
+}
+
+bool UFractureToolGenerateAsset::CanExecute() const
+{
+	// We can execute this command if static meshes are selected.
+	return IsStaticMeshSelected();
+}
+
+void UFractureToolGenerateAsset::Execute(TWeakPtr<FFractureEditorModeToolkit> InToolkit)
+{
+	Toolkit = InToolkit;
+	
+	USelection* SelectionSet = GEditor->GetSelectedActors();
+
+	TArray<AActor*> SelectedActors;
+	SelectedActors.Reserve(SelectionSet->Num());
+
+	FScopedTransaction Transaction(LOCTEXT("GenerateAsset", "Generate Geometry Collection Asset"));
+
+	SelectionSet->GetSelectedObjects(SelectedActors);
+
+	OpenGenerateAssetDialog(SelectedActors);
+}
+
+void UFractureToolGenerateAsset::OpenGenerateAssetDialog(TArray<AActor*>& Actors)
+{
+	TSharedPtr<SWindow> PickAssetPathWindow;
+
+	SAssignNew(PickAssetPathWindow, SWindow)
+		.Title(LOCTEXT("SelectPath", "Select Path"))
+		.ToolTipText(LOCTEXT("SelectPathTooltip", "Select the path where the Geometry Collection will be created at"))
+		.ClientSize(FVector2D(400, 400));
+
+	// NOTE - the parent window has to completely exist before this one does so the parent gets set properly.
+	// This is why we do not just put this in the Contents()[ ... ] of the Window above.
+	TSharedPtr<SCreateAssetFromObject> CreateAssetDialog;
+	PickAssetPathWindow->SetContent(
+		SAssignNew(CreateAssetDialog, SCreateAssetFromObject, PickAssetPathWindow)
+		.AssetFilenameSuffix(TEXT("GeometryCollection"))
+		.HeadingText(LOCTEXT("CreateGeometryCollection_Heading", "Geometry Collection Name"))
+		.CreateButtonText(LOCTEXT("CreateGeometryCollection_ButtonLabel", "Create Geometry Collection"))
+		.OnCreateAssetAction(FOnPathChosen::CreateUObject(this, &UFractureToolGenerateAsset::OnGenerateAssetPathChosen, Actors))
+	);
+
+	TSharedPtr<SWindow> RootWindow = FGlobalTabmanager::Get()->GetRootWindow();
+	if (RootWindow.IsValid())
+	{
+		FSlateApplication::Get().AddWindowAsNativeChild(PickAssetPathWindow.ToSharedRef(), RootWindow.ToSharedRef());
+	}
+	else
+	{
+		FSlateApplication::Get().AddWindow(PickAssetPathWindow.ToSharedRef());
+	}
+
+}
+
+void UFractureToolGenerateAsset::OnGenerateAssetPathChosen(const FString& InAssetPath, TArray<AActor*> Actors)
+{
+		
+	UGeometryCollectionComponent* GeometryCollectionComponent = nullptr;//  = Cast<UGeometryCollectionComponent>(FractureContext.OriginalPrimitiveComponent);
+
+	if (Actors.Num() > 0)
+	{
+		AActor* FirstActor = Actors[0];
+
+		AGeometryCollectionActor* GeometryCollectionActor = Cast<AGeometryCollectionActor>(FirstActor);
+		GeometryCollectionActor = ConvertStaticMeshToGeometryCollection(InAssetPath, Actors);
+
+		GeometryCollectionComponent = GeometryCollectionActor->GetGeometryCollectionComponent();
+
+		FScopedColorEdit EditBoneColor = GeometryCollectionComponent->EditBoneSelection();
+		EditBoneColor.SetShowBoneColors(true);
+
+		// Move GC actor to source actors position and remove source actor from scene
+		const FVector ActorLocation(FirstActor->GetActorLocation());
+		GeometryCollectionActor->SetActorLocation(ActorLocation);
+
+		// Clear selection of mesh actor used to make GC before selecting, will cause details pane to not display geometry collection details.
+		GEditor->SelectNone(true, true, false);
+		GEditor->SelectActor(GeometryCollectionActor, true, true);
+
+		EditBoneColor.SelectBones(GeometryCollection::ESelectionMode::AllGeometry);
+
+		if (Toolkit.IsValid())
+		{
+			TSharedPtr<FFractureEditorModeToolkit> SharedToolkit(Toolkit.Pin());
+			SharedToolkit->SetOutlinerComponents({ GeometryCollectionComponent });
+			SharedToolkit->SetBoneSelection(GeometryCollectionComponent, EditBoneColor.GetSelectedBones(), true);
+		}
+		
+		GeometryCollectionComponent->MarkRenderDynamicDataDirty();
+		GeometryCollectionComponent->MarkRenderStateDirty();
+
+		for (AActor* Actor : Actors)
+		{
+			Actor->Destroy();
+		}
+	}
+}
+
+AGeometryCollectionActor* UFractureToolGenerateAsset::ConvertStaticMeshToGeometryCollection(const FString& InAssetPath, TArray<AActor*>& Actors)
+{
+	ensure(Actors.Num() > 0);
+	AActor* FirstActor = Actors[0];
+	const FString& Name = FirstActor->GetActorLabel();
+	const FVector FirstActorLocation(FirstActor->GetActorLocation());
+
+
+	AGeometryCollectionActor* NewActor = CreateNewGeometryActor(InAssetPath, FTransform(), true);
+
+	FGeometryCollectionEdit GeometryCollectionEdit = NewActor->GetGeometryCollectionComponent()->EditRestCollection(GeometryCollection::EEditUpdate::RestPhysicsDynamic);
+	UGeometryCollection* FracturedGeometryCollection = GeometryCollectionEdit.GetRestCollection();
+
+	for (AActor* Actor : Actors)
+	{
+		const FTransform ActorTransform(Actor->GetTransform());
+		const FVector ActorOffset(Actor->GetActorLocation() - FirstActor->GetActorLocation());
+
+		check(FracturedGeometryCollection);
+
+		TArray<UStaticMeshComponent*> StaticMeshComponents;
+		Actor->GetComponents<UStaticMeshComponent>(StaticMeshComponents, true);
+		for (int32 ii = 0, ni = StaticMeshComponents.Num(); ii < ni; ++ii)
+		{
+			// We're partial to static mesh components, here
+			UStaticMeshComponent* StaticMeshComponent = StaticMeshComponents[ii];
+			if (StaticMeshComponent != nullptr)
+			{
+				UStaticMesh* ComponentStaticMesh = StaticMeshComponent->GetStaticMesh();
+				FTransform ComponentTransform(StaticMeshComponent->GetComponentTransform());
+				ComponentTransform.SetTranslation((ComponentTransform.GetTranslation() - ActorTransform.GetTranslation()) + ActorOffset);
+
+				// Record the contributing source on the asset.
+				FSoftObjectPath SourceSoftObjectPath(ComponentStaticMesh);
+				TArray<UMaterialInterface*> SourceMaterials = StaticMeshComponent->GetMaterials();
+				FracturedGeometryCollection->GeometrySource.Add({ SourceSoftObjectPath, ComponentTransform, SourceMaterials });
+
+				FGeometryCollectionConversion::AppendStaticMesh(ComponentStaticMesh, SourceMaterials, ComponentTransform, FracturedGeometryCollection, true);
+			}
+		}
+
+		FracturedGeometryCollection->InitializeMaterials();
+	}
+	AddSingleRootNodeIfRequired(FracturedGeometryCollection);
+
+	return NewActor;
+}
+
+class AGeometryCollectionActor* UFractureToolGenerateAsset::CreateNewGeometryActor(const FString& InAssetPath, const FTransform& Transform, bool AddMaterials /*= false*/)
+{
+
+	FString UniquePackageName = InAssetPath;
+	FString UniqueAssetName = FPackageName::GetLongPackageAssetName(InAssetPath);
+
+	FAssetToolsModule& AssetToolsModule = FModuleManager::LoadModuleChecked<FAssetToolsModule>("AssetTools");
+	AssetToolsModule.Get().CreateUniqueAssetName(UniquePackageName, TEXT(""), UniquePackageName, UniqueAssetName);
+
+	UPackage* Package = CreatePackage(*UniquePackageName);
+	UGeometryCollection* InGeometryCollection = static_cast<UGeometryCollection*>(NewObject<UGeometryCollection>(Package, UGeometryCollection::StaticClass(), FName(*UniqueAssetName), RF_Transactional | RF_Public | RF_Standalone));
+
+	// Create the new Geometry Collection actor
+	AGeometryCollectionActor* NewActor = Cast<AGeometryCollectionActor>(AddActor(GetSelectedLevel(), AGeometryCollectionActor::StaticClass()));
+	check(NewActor->GetGeometryCollectionComponent());
+
+	// Set the Geometry Collection asset in the new actor
+	NewActor->GetGeometryCollectionComponent()->SetRestCollection(InGeometryCollection);
+
+	// copy transform of original static mesh actor to this new actor
+	NewActor->SetActorLabel(UniqueAssetName);
+	NewActor->SetActorTransform(Transform);
+
+	// Mark relevant stuff dirty
+	FAssetRegistryModule::AssetCreated(InGeometryCollection);
+	InGeometryCollection->MarkPackageDirty();
+	Package->SetDirtyFlag(true);
+
+	return NewActor;
+}
+
+ULevel* UFractureToolGenerateAsset::GetSelectedLevel()
+{
+	USelection* SelectedActors = GEditor->GetSelectedActors();
+	TArray<ULevel*> UniqueLevels;
+	for (FSelectionIterator Iter(*SelectedActors); Iter; ++Iter)
+	{
+		AActor* Actor = Cast<AActor>(*Iter);
+		if (Actor)
+		{
+			UniqueLevels.AddUnique(Actor->GetLevel());
+		}
+	}
+	check(UniqueLevels.Num() == 1);
+	return UniqueLevels[0];
+}
+
+AActor* UFractureToolGenerateAsset::AddActor(ULevel* InLevel, UClass* Class)
+{
+	check(Class);
+
+	UWorld* World = InLevel->OwningWorld;
+	ULevel* DesiredLevel = InLevel;
+
+	// Transactionally add the actor.
+	AActor* Actor = NULL;
+	{
+		FScopedTransaction Transaction(NSLOCTEXT("UnrealEd", "AddActor", "Add Actor"));
+
+		FActorSpawnParameters SpawnInfo;
+		SpawnInfo.OverrideLevel = DesiredLevel;
+		SpawnInfo.SpawnCollisionHandlingOverride = ESpawnActorCollisionHandlingMethod::AlwaysSpawn;
+		SpawnInfo.ObjectFlags = RF_Transactional;
+		const auto Location = FVector(0);
+		const auto Rotation = FTransform(FVector(0)).GetRotation().Rotator();
+		Actor = World->SpawnActor(Class, &Location, &Rotation, SpawnInfo);
+
+		check(Actor);
+		Actor->InvalidateLightingCache();
+		Actor->PostEditMove(true);
+	}
+
+	// If this actor is part of any layers (set in its default properties), add them into the visible layers list.
+	ULayersSubsystem* Layers = GEditor->GetEditorSubsystem<ULayersSubsystem>();
+	Layers->SetLayersVisibility(Actor->Layers, true);
+
+	// Clean up.
+	Actor->MarkPackageDirty();
+	ULevel::LevelDirtiedEvent.Broadcast();
+
+	return Actor;
+}
+
+
+FText UFractureToolResetAsset::GetDisplayText() const
+{
+	return FText(NSLOCTEXT("Fracture", "FractureToolReset", "Reset"));
+}
+
+FText UFractureToolResetAsset::GetTooltipText() const
+{
+	return FText(NSLOCTEXT("Fracture", "FractureToolResetTooltip", "Reset Geometry Collection Asset to its initial unfractured state."));
+}
+
+FSlateIcon UFractureToolResetAsset::GetToolIcon() const
+{
+	return FSlateIcon("FractureEditorStyle", "FractureEditor.ResetAsset");
+}
+
+void UFractureToolResetAsset::RegisterUICommand(FFractureEditorCommands* BindingContext)
+{
+	UI_COMMAND_EXT(BindingContext, UICommandInfo, "ResetAsset", "Reset", "Reset", EUserInterfaceActionType::ToggleButton, FInputChord());
+	BindingContext->ResetAsset = UICommandInfo;
+}
+
+bool UFractureToolResetAsset::CanExecute() const
+{
+	return IsGeometryCollectionSelected();
+}
+
+void UFractureToolResetAsset::Execute(TWeakPtr<FFractureEditorModeToolkit> InToolkit)
+{
+	if (!InToolkit.IsValid())
+	{
+		return;
+	}
+
+	TSet<UGeometryCollectionComponent*> GeomCompSelection;
+	GetSelectedGeometryCollectionComponents(GeomCompSelection);
+	for (UGeometryCollectionComponent* GeometryCollectionComponent : GeomCompSelection)
+	{
+		FGeometryCollectionEdit GeometryCollectionEdit = GeometryCollectionComponent->EditRestCollection();
+		if (UGeometryCollection* GeometryCollectionObject = GeometryCollectionEdit.GetRestCollection())
+		{
+
+			TSharedPtr<FGeometryCollection, ESPMode::ThreadSafe> GeometryCollectionPtr = GeometryCollectionObject->GetGeometryCollection();
+			if (FGeometryCollection* GeometryCollection = GeometryCollectionPtr.Get())
+			{
+				GeometryCollectionObject->Reset();
+
+				// Rebuild Collection from recorded source assets.
+				for (const FGeometryCollectionSource& Source : GeometryCollectionObject->GeometrySource)
+				{
+					const UObject* SourceMesh = Source.SourceGeometryObject.TryLoad();
+					if (const UStaticMesh* SourceStaticMesh = Cast<UStaticMesh>(SourceMesh))
+					{
+						FGeometryCollectionConversion::AppendStaticMesh(SourceStaticMesh, Source.SourceMaterial, Source.LocalTransform, GeometryCollectionObject, true);
+					}
+					else if (const USkeletalMesh* SourceSkeletalMesh = Cast<USkeletalMesh>(SourceMesh))
+					{
+						// #todo (bmiller) Once we've settled on the right approach with static meshes, we'll need to apply the same strategy to skeletal mesh reconstruction.
+						// FGeometryCollectionConversion::AppendSkeletalMesh(SourceSkeletalMesh, Source.SourceMaterial, Source.LocalTransform, GeometryCollectionObject, true);
+					}
+				}
+
+				GeometryCollectionObject->InitializeMaterials();
+
+				FGeometryCollectionClusteringUtility::UpdateHierarchyLevelOfChildren(GeometryCollection, -1);
+				AddSingleRootNodeIfRequired(GeometryCollectionObject);
+				GeometryCollectionComponent->MarkRenderDynamicDataDirty();
+				GeometryCollectionComponent->MarkRenderStateDirty();
+			}
+			GeometryCollectionObject->MarkPackageDirty();
+		}
+	}
+	InToolkit.Pin()->SetOutlinerComponents(GeomCompSelection.Array());
+}
+
+#undef LOCTEXT_NAMESPACE
diff --git a/Engine/Plugins/Experimental/ChaosEditor/Source/FractureEditor/Private/FractureToolGenerators.h b/Engine/Plugins/Experimental/ChaosEditor/Source/FractureEditor/Private/FractureToolGenerators.h
new file mode 100644
--- /dev/null
+++ b/Engine/Plugins/Experimental/ChaosEditor/Source/FractureEditor/Private/FractureToolGenerators.h
@@ -0,0 +1,60 @@
+// Copyright Epic Games, Inc. All Rights Reserved.
+
+#pragma once
+
+#include "FractureTool.h"
+
+#include "FractureToolGenerators.generated.h"
+
+class AGeometryCollectionActor;
+class UGeometryCollection;
+
+UCLASS()
+class UFractureToolGenerateAsset : public UFractureActionTool
+{
+public:
+	GENERATED_BODY()
+
+	UFractureToolGenerateAsset(const FObjectInitializer& ObjInit) 
+		: Super(ObjInit)
+	{}
+
+	// UFractureActionTool Interface
+	virtual FText GetDisplayText() const override;
+	virtual FText GetTooltipText() const override;
+	virtual FSlateIcon GetToolIcon() const override;
+	virtual void RegisterUICommand(FFractureEditorCommands* BindingContext) override;
+	virtual void Execute(TWeakPtr<FFractureEditorModeToolkit> InToolkit) override;
+	virtual bool CanExecute() const override;
+
+private:
+	void OpenGenerateAssetDialog(TArray<AActor*>& Actors);
+	void OnGenerateAssetPathChosen(const FString& InAssetPath, TArray<AActor*> Actors);
+	AGeometryCollectionActor* ConvertStaticMeshToGeometryCollection(const FString& InAssetPath, TArray<AActor*>& Actors);
+	class AGeometryCollectionActor* CreateNewGeometryActor(const FString& Name, const FTransform& Transform, bool AddMaterials /*= false*/);
+	static ULevel* GetSelectedLevel();
+	static AActor* AddActor(ULevel* InLevel, UClass* Class);
+
+private:
+	TWeakPtr<FFractureEditorModeToolkit> Toolkit;
+
+};
+
+
+UCLASS()
+class UFractureToolResetAsset : public UFractureToolGenerateAsset
+{
+public:
+	GENERATED_BODY()
+
+	UFractureToolResetAsset(const FObjectInitializer& ObjInit) : Super(ObjInit) {}
+
+	// UFractureActionTool Interface
+	virtual FText GetDisplayText() const override;
+	virtual FText GetTooltipText() const override;
+	virtual FSlateIcon GetToolIcon() const override;
+	virtual void RegisterUICommand(FFractureEditorCommands* BindingContext) override;
+	virtual void Execute(TWeakPtr<FFractureEditorModeToolkit> InToolkit) override;
+	virtual bool CanExecute() const override;
+};
+
diff --git a/Engine/Plugins/Experimental/ChaosEditor/Source/FractureEditor/Private/FractureToolPlaneCut.cpp b/Engine/Plugins/Experimental/ChaosEditor/Source/FractureEditor/Private/FractureToolPlaneCut.cpp
--- a/Engine/Plugins/Experimental/ChaosEditor/Source/FractureEditor/Private/FractureToolPlaneCut.cpp
+++ b/Engine/Plugins/Experimental/ChaosEditor/Source/FractureEditor/Private/FractureToolPlaneCut.cpp
@@ -9,26 +9,6 @@
 
 #define LOCTEXT_NAMESPACE "FracturePlanar"
 
-void UFracturePlaneCutSettings::PostEditChangeProperty(FPropertyChangedEvent& PropertyChangedEvent)
-{
-	if (OwnerTool != nullptr)
-	{
-		OwnerTool->PostEditChangeProperty(PropertyChangedEvent);
-	}
-	Super::PostEditChangeProperty(PropertyChangedEvent);
-}
-
-void UFracturePlaneCutSettings::PostEditChangeChainProperty(struct FPropertyChangedChainEvent& PropertyChangedEvent)
-{
-	if (OwnerTool != nullptr)
-	{
-		OwnerTool->PostEditChangeChainProperty(PropertyChangedEvent);
-	}
-	Super::PostEditChangeChainProperty(PropertyChangedEvent);
-}
-
-
-
 
 UFractureToolPlaneCut::UFractureToolPlaneCut(const FObjectInitializer& ObjInit) 
 	: Super(ObjInit) 
@@ -89,24 +69,22 @@ void UFractureToolPlaneCut::Render(const FSceneView* View, FViewport* Viewport,
 }
 
 TArray<UObject*> UFractureToolPlaneCut::GetSettingsObjects() const
-{
+ {
 	TArray<UObject*> Settings;
-	Settings.Add(GetMutableDefault<UFractureCommonSettings>());
-	Settings.Add(GetMutableDefault<UFracturePlaneCutSettings>());
+	Settings.Add(CutterSettings);
+	Settings.Add(PlaneCutSettings);
 	return Settings;
 }
 
 void UFractureToolPlaneCut::FractureContextChanged()
 {
-	const UFractureCommonSettings* LocalCommonSettings = GetDefault<UFractureCommonSettings>();
-
-	TArray<FFractureContext> FractureContexts;
-	FFractureEditorModeToolkit::GetFractureContexts(FractureContexts);
+	TArray<FFractureToolContext> FractureContexts;
+	GetFractureContexts(FractureContexts);
 
 	RenderCuttingPlanesTransforms.Empty();
 
 	RenderCuttingPlaneSize = FLT_MAX;
-	for (FFractureContext& FractureContext : FractureContexts)
+	for (FFractureToolContext& FractureContext : FractureContexts)
 	{
 		// Move the local bounds to the actor so we weill draw in the correct location
 		FractureContext.Bounds = FractureContext.Bounds.TransformBy(FractureContext.Transform);
@@ -119,10 +97,8 @@ void UFractureToolPlaneCut::FractureContextChanged()
 	}
 }
 
-void UFractureToolPlaneCut::ExecuteFracture(const FFractureContext& FractureContext)
+void UFractureToolPlaneCut::ExecuteFracture(const FFractureToolContext& FractureContext)
 {
-	const UFractureCommonSettings* LocalCommonSettings = GetDefault<UFractureCommonSettings>();
-
 	if (FractureContext.FracturedGeometryCollection != nullptr)
 	{
 		TSharedPtr<FGeometryCollection, ESPMode::ThreadSafe> GeometryCollectionPtr = FractureContext.FracturedGeometryCollection->GetGeometryCollection();
@@ -148,12 +124,12 @@ void UFractureToolPlaneCut::ExecuteFracture(const FFractureContext& FractureCont
 
 			FInternalSurfaceMaterials InternalSurfaceMaterials;
 			FNoiseSettings NoiseSettings;
-			if (LocalCommonSettings->Amplitude > 0.0f)
+			if (CutterSettings->Amplitude > 0.0f)
 			{
-				NoiseSettings.Amplitude = LocalCommonSettings->Amplitude;
-				NoiseSettings.Frequency = LocalCommonSettings->Frequency;
-				NoiseSettings.Octaves = LocalCommonSettings->OctaveNumber;
-				NoiseSettings.PointSpacing = LocalCommonSettings->SurfaceResolution;
+				NoiseSettings.Amplitude = CutterSettings->Amplitude;
+				NoiseSettings.Frequency = CutterSettings->Frequency;
+				NoiseSettings.Octaves = CutterSettings->OctaveNumber;
+				NoiseSettings.PointSpacing = CutterSettings->SurfaceResolution;
 				InternalSurfaceMaterials.NoiseSettings = NoiseSettings;
 			}
 
@@ -162,29 +138,14 @@ void UFractureToolPlaneCut::ExecuteFracture(const FFractureContext& FractureCont
 	}
 }
 
-bool UFractureToolPlaneCut::CanExecuteFracture() const
-{
-	return FFractureEditorModeToolkit::IsLeafBoneSelected();
-}
-
-#if WITH_EDITOR
-void UFractureToolPlaneCut::PostEditChangeChainProperty(struct FPropertyChangedChainEvent& PropertyChangedEvent)
+void UFractureToolPlaneCut::GenerateSliceTransforms(const FFractureToolContext& Context, TArray<FTransform>& CuttingPlaneTransforms)
 {
-	FractureContextChanged();
-}
-#endif
-
-
-void UFractureToolPlaneCut::GenerateSliceTransforms(const FFractureContext& Context, TArray<FTransform>& CuttingPlaneTransforms)
-{
-	const UFracturePlaneCutSettings* LocalCutSettings = GetMutableDefault<UFracturePlaneCutSettings>();
-
 	FRandomStream RandStream(Context.RandomSeed);
 
 	const FVector Extent(Context.Bounds.Max - Context.Bounds.Min);
 
-	CuttingPlaneTransforms.Reserve(CuttingPlaneTransforms.Num() + LocalCutSettings->NumberPlanarCuts);
-	for (int32 ii = 0; ii < LocalCutSettings->NumberPlanarCuts; ++ii)
+	CuttingPlaneTransforms.Reserve(CuttingPlaneTransforms.Num() + PlaneCutSettings->NumberPlanarCuts);
+	for (int32 ii = 0; ii < PlaneCutSettings->NumberPlanarCuts; ++ii)
 	{
 		FVector Position(Context.Bounds.Min + FVector(RandStream.FRand(), RandStream.FRand(), RandStream.FRand()) * Extent);
 		CuttingPlaneTransforms.Emplace(FTransform(FRotator(RandStream.FRand() * 360.0f, RandStream.FRand() * 360.0f, 0.0f), Position));
diff --git a/Engine/Plugins/Experimental/ChaosEditor/Source/FractureEditor/Private/FractureToolPlaneCut.h b/Engine/Plugins/Experimental/ChaosEditor/Source/FractureEditor/Private/FractureToolPlaneCut.h
--- a/Engine/Plugins/Experimental/ChaosEditor/Source/FractureEditor/Private/FractureToolPlaneCut.h
+++ b/Engine/Plugins/Experimental/ChaosEditor/Source/FractureEditor/Private/FractureToolPlaneCut.h
@@ -2,23 +2,20 @@
 
 #pragma once
 
-#include "FractureTool.h"
+#include "FractureToolCutter.h"
 
 #include "FractureToolPlaneCut.generated.h"
 
+
 UCLASS(config = EditorPerProjectUserSettings)
-class UFracturePlaneCutSettings
-	: public UObject
+class UFracturePlaneCutSettings : public UFractureToolSettings
 {
-	GENERATED_BODY()
 public:
-#if WITH_EDITOR
-	virtual void PostEditChangeProperty(FPropertyChangedEvent& PropertyChangedEvent) override;
-	virtual void PostEditChangeChainProperty(struct FPropertyChangedChainEvent& PropertyChangedEvent) override;
-#endif
+	GENERATED_BODY()
 
-	UFracturePlaneCutSettings()
-	: NumberPlanarCuts(3) {}
+	UFracturePlaneCutSettings(const FObjectInitializer& ObjInit)
+		: Super(ObjInit)
+		, NumberPlanarCuts(3) {}
 
 	/** Number of Clusters - Cluster Voronoi Method */
 	UPROPERTY(EditAnywhere, Category = PlaneCut, meta = (DisplayName = "Number of Cuts", UIMin = "1", UIMax = "20", ClampMin = "1"))
@@ -28,18 +25,16 @@ public:
 	UPROPERTY(EditAnywhere, Category = PlaneCut, meta = (DisplayName = "Reference Actor"))
 	TLazyObjectPtr<AActor> ReferenceActor;
 
-	UPROPERTY()
-	UFractureTool *OwnerTool;
-
 };
 
-UCLASS(DisplayName="Plane Cut Tool", Category="FractureTools")
-class UFractureToolPlaneCut : public UFractureTool
+
+UCLASS(DisplayName = "Plane Cut Tool", Category = "FractureTools")
+class UFractureToolPlaneCut : public UFractureToolCutterBase
 {
 public:
 	GENERATED_BODY()
 
-	UFractureToolPlaneCut(const FObjectInitializer& ObjInit);//  : Super(ObjInit) {}
+	UFractureToolPlaneCut(const FObjectInitializer& ObjInit);
 
 	// UFractureTool Interface
 	virtual FText GetDisplayText() const override;
@@ -48,25 +43,22 @@ public:
 
 	void Render(const FSceneView* View, FViewport* Viewport, FPrimitiveDrawInterface* PDI) override;
 
-	virtual void RegisterUICommand( FFractureEditorCommands* BindingContext );
-	
-	virtual TArray<UObject*> GetSettingsObjects() const;// { return TArray<UObject*>(); }
+	virtual void RegisterUICommand(FFractureEditorCommands* BindingContext) override;
 
-	virtual void FractureContextChanged() override;
-	virtual void ExecuteFracture(const FFractureContext& FractureContext) override;
-	virtual bool CanExecuteFracture() const override;
+	virtual TArray<UObject*> GetSettingsObjects() const override;
 
-#if WITH_EDITOR
-	virtual void PostEditChangeChainProperty(struct FPropertyChangedChainEvent& PropertyChangedEvent) override;
-#endif
+	virtual void FractureContextChanged() override;
+	virtual void ExecuteFracture(const FFractureToolContext& FractureContext) override;
 
 private:
 	// Slicing
 	UPROPERTY(EditAnywhere, Category = Slicing)
 	UFracturePlaneCutSettings* PlaneCutSettings;
 
-	void GenerateSliceTransforms(const FFractureContext& Context, TArray<FTransform>& CuttingPlaneTransforms);
+	void GenerateSliceTransforms(const FFractureToolContext& Context, TArray<FTransform>& CuttingPlaneTransforms);
 
 	float RenderCuttingPlaneSize;
 	TArray<FTransform> RenderCuttingPlanesTransforms;
-};
\ No newline at end of file
+};
+
+
diff --git a/Engine/Plugins/Experimental/ChaosEditor/Source/FractureEditor/Private/FractureToolRadial.cpp b/Engine/Plugins/Experimental/ChaosEditor/Source/FractureEditor/Private/FractureToolRadial.cpp
--- a/Engine/Plugins/Experimental/ChaosEditor/Source/FractureEditor/Private/FractureToolRadial.cpp
+++ b/Engine/Plugins/Experimental/ChaosEditor/Source/FractureEditor/Private/FractureToolRadial.cpp
@@ -6,29 +6,12 @@
 
 #define LOCTEXT_NAMESPACE "FractureRadial"
 
-void UFractureRadialSettings::PostEditChangeProperty(FPropertyChangedEvent& PropertyChangedEvent)
-{
-	if (OwnerTool != nullptr)
-	{
-		OwnerTool->PostEditChangeProperty(PropertyChangedEvent);
-	}
-	Super::PostEditChangeProperty(PropertyChangedEvent);
-}
-
-void UFractureRadialSettings::PostEditChangeChainProperty(struct FPropertyChangedChainEvent& PropertyChangedEvent)
-{
-	if (OwnerTool != nullptr)
-	{
-		OwnerTool->PostEditChangeChainProperty(PropertyChangedEvent);
-	}
-	Super::PostEditChangeChainProperty(PropertyChangedEvent);
-}
-
 
 UFractureToolRadial::UFractureToolRadial(const FObjectInitializer& ObjInit) 
 	: Super(ObjInit) 
 {
-	GetMutableDefault<UFractureRadialSettings>()->OwnerTool = this;
+	RadialSettings = NewObject<UFractureRadialSettings>(GetTransientPackage(), UFractureRadialSettings::StaticClass());
+	RadialSettings->OwnerTool = this;
 }
 
 FText UFractureToolRadial::GetDisplayText() const
@@ -55,35 +38,33 @@ void UFractureToolRadial::RegisterUICommand( FFractureEditorCommands* BindingCon
 TArray<UObject*> UFractureToolRadial::GetSettingsObjects() const 
 { 
 	TArray<UObject*> Settings; 
-	Settings.Add(GetMutableDefault<UFractureCommonSettings>());
-	Settings.Add(GetMutableDefault<UFractureRadialSettings>());
+	Settings.Add(CutterSettings);
+	Settings.Add(RadialSettings);
 	return Settings;
 }
 
-void UFractureToolRadial::GenerateVoronoiSites(const FFractureContext &Context, TArray<FVector>& Sites)
+void UFractureToolRadial::GenerateVoronoiSites(const FFractureToolContext&Context, TArray<FVector>& Sites)
 {
- 	const UFractureRadialSettings* FractureSettings = GetMutableDefault<UFractureRadialSettings>();
-
-	float RadialStep = FractureSettings->Radius / FractureSettings->RadialSteps;
+ 	float RadialStep = RadialSettings->Radius / RadialSettings->RadialSteps;
 
-	const FVector Center(Context.Bounds.GetCenter() + FractureSettings->Center);
+	const FVector Center(Context.Bounds.GetCenter() + RadialSettings->Center);
 
 	FRandomStream RandStream(Context.RandomSeed);
-	FVector UpVector(FractureSettings->Normal);
+	FVector UpVector(RadialSettings->Normal);
 	UpVector.Normalize();
 	FVector PerpVector(UpVector[2], UpVector[0], UpVector[1]);
 
-	for (int32 ii = 1; ii < FractureSettings->RadialSteps; ++ii)
+	for (int32 ii = 1; ii < RadialSettings->RadialSteps; ++ii)
 	{
 		FVector PositionVector(PerpVector * RadialStep * ii);
 
-		float AngularStep = 360.f / FractureSettings->AngularSteps;
-		PositionVector = PositionVector.RotateAngleAxis(FractureSettings->AngleOffset * ii, UpVector);
+		float AngularStep = 360.f / RadialSettings->AngularSteps;
+		PositionVector = PositionVector.RotateAngleAxis(RadialSettings->AngleOffset * ii, UpVector);
 
-		for (int32 kk = 0; kk < FractureSettings->AngularSteps; ++kk)
+		for (int32 kk = 0; kk < RadialSettings->AngularSteps; ++kk)
 		{
 			PositionVector = PositionVector.RotateAngleAxis(AngularStep , UpVector);
-			Sites.Emplace(Center + PositionVector + (RandStream.VRand() * RandStream.FRand() * FractureSettings->Variability));
+			Sites.Emplace(Center + PositionVector + (RandStream.VRand() * RandStream.FRand() * RadialSettings->Variability));
 		}
 	}
 }
diff --git a/Engine/Plugins/Experimental/ChaosEditor/Source/FractureEditor/Private/FractureToolRadial.h b/Engine/Plugins/Experimental/ChaosEditor/Source/FractureEditor/Private/FractureToolRadial.h
--- a/Engine/Plugins/Experimental/ChaosEditor/Source/FractureEditor/Private/FractureToolRadial.h
+++ b/Engine/Plugins/Experimental/ChaosEditor/Source/FractureEditor/Private/FractureToolRadial.h
@@ -2,33 +2,29 @@
 
 #pragma once
 
-#include "FractureToolVoronoiBase.h"
+#include "FractureTool.h"
+#include "FractureToolCutter.h"
 
 #include "FractureToolRadial.generated.h"
 
 
-UCLASS()
-class UFractureRadialSettings
-	: public UObject
+UCLASS(config = EditorPerProjectUserSettings)
+class UFractureRadialSettings : public UFractureToolSettings
 {
-	GENERATED_BODY()
 public:
+	GENERATED_BODY()
 
-	UFractureRadialSettings()
-	: Center(FVector(0,0,0))
-	, Normal(FVector(0, 0, 1))
-	, Radius(50.0f)
-	, AngularSteps(5)
-	, RadialSteps(5)
-	, AngleOffset(0.0f)
-	, Variability(0.0f)
+	UFractureRadialSettings(const FObjectInitializer& ObjInit)
+		: Super(ObjInit)
+		, Center(FVector(0,0,0))
+		, Normal(FVector(0, 0, 1))
+		, Radius(50.0f)
+		, AngularSteps(5)
+		, RadialSteps(5)
+		, AngleOffset(0.0f)
+		, Variability(0.0f)
 	{}
 
-#if WITH_EDITOR
-	virtual void PostEditChangeProperty(FPropertyChangedEvent& PropertyChangedEvent) override;
-	virtual void PostEditChangeChainProperty(struct FPropertyChangedChainEvent& PropertyChangedEvent) override;
-#endif
-
 	/** Center of generated pattern */
 	UPROPERTY(EditAnywhere, Category = RadialVoronoi)
 	FVector Center;
@@ -56,31 +52,30 @@ public:
 	/** Randomness of sites distribution */
 	UPROPERTY(EditAnywhere, Category = RadialVoronoi, meta = (DisplayName = "Variability", UIMin = "0.0", ClampMin = "0.0"))
 	float Variability;
-
-	UPROPERTY()
-	UFractureTool *OwnerTool;
 };
 
 
 UCLASS(DisplayName="Radial Voronoi", Category="FractureTools")
-class UFractureToolRadial : public UFractureToolVoronoiBase
+class UFractureToolRadial : public UFractureToolVoronoiCutterBase
 {
 public:
 	GENERATED_BODY()
 
-	UFractureToolRadial(const FObjectInitializer& ObjInit);//  : Super(ObjInit) {}
+	UFractureToolRadial(const FObjectInitializer& ObjInit);
 
 	// UFractureTool Interface
 	virtual FText GetDisplayText() const override;
 	virtual FText GetTooltipText() const override;
 	virtual FSlateIcon GetToolIcon() const override;
-	virtual TArray<UObject*> GetSettingsObjects() const override;// { return TArray<UObject*>(); }
+	virtual TArray<UObject*> GetSettingsObjects() const override;
+
+	virtual void RegisterUICommand( FFractureEditorCommands* BindingContext ) override;
+
+	// Radial Voronoi Fracture Input Settings
+	UPROPERTY(EditAnywhere, Category = Uniform)
+	UFractureRadialSettings* RadialSettings;
 
-	virtual void RegisterUICommand( FFractureEditorCommands* BindingContext );
-	
-	// virtual void ExecuteFracture() {}
-	// virtual bool CanExecuteFracture() { return true; }
 protected:
-	void GenerateVoronoiSites(const FFractureContext &Context, TArray<FVector>& Sites) override;
+	void GenerateVoronoiSites(const FFractureToolContext &Context, TArray<FVector>& Sites) override;
 
 };
\ No newline at end of file
diff --git a/Engine/Plugins/Experimental/ChaosEditor/Source/FractureEditor/Private/FractureToolSelectors.cpp b/Engine/Plugins/Experimental/ChaosEditor/Source/FractureEditor/Private/FractureToolSelectors.cpp
new file mode 100644
--- /dev/null
+++ b/Engine/Plugins/Experimental/ChaosEditor/Source/FractureEditor/Private/FractureToolSelectors.cpp
@@ -0,0 +1,216 @@
+// Copyright Epic Games, Inc. All Rights Reserved.
+
+#include "FractureToolSelectors.h"
+
+#include "Editor.h"
+#include "Engine/Selection.h"
+
+#include "GeometryCollection/GeometryCollectionComponent.h"
+#include "GeometryCollection/GeometryCollection.h"
+#include "FractureSelectionTools.h"
+
+
+#define LOCTEXT_NAMESPACE "FractureToolSelectionOps"
+
+
+FText UFractureToolSelectAll::GetDisplayText() const
+{
+	return FText(NSLOCTEXT("Fracture", "FractureToolSelectAll", "Select All"));
+}
+
+FText UFractureToolSelectAll::GetTooltipText() const
+{
+	return FText(NSLOCTEXT("Fracture", "FractureToolSelectAllTooltip", "Selects all Bones in the GeometryCollection"));
+}
+
+FSlateIcon UFractureToolSelectAll::GetToolIcon() const
+{
+	return FSlateIcon("FractureEditorStyle", "FractureEditor.SelectAll");
+}
+
+void UFractureToolSelectAll::RegisterUICommand(FFractureEditorCommands* BindingContext)
+{
+	UI_COMMAND_EXT(BindingContext, UICommandInfo, "SelectAll", "Select All", "Selects all Bones in the GeometryCollection.", EUserInterfaceActionType::Button, FInputChord(EModifierKey::Control | EModifierKey::Shift, EKeys::A));
+	BindingContext->SelectAll = UICommandInfo;
+}
+
+void UFractureToolSelectAll::Execute(TWeakPtr<FFractureEditorModeToolkit> InToolkit)
+{
+	if (InToolkit.IsValid())
+	{
+		SelectByMode(InToolkit.Pin().Get(), GeometryCollection::ESelectionMode::AllGeometry);
+	}
+}
+
+void UFractureToolSelectAll::SelectByMode(FFractureEditorModeToolkit* InToolkit, GeometryCollection::ESelectionMode SelectionMode)
+{
+	USelection* SelectionSet = GEditor->GetSelectedActors();
+
+	TArray<AActor*> SelectedActors;
+	SelectedActors.Reserve(SelectionSet->Num());
+	SelectionSet->GetSelectedObjects(SelectedActors);
+
+	for (AActor* Actor : SelectedActors)
+	{
+		TInlineComponentArray<UGeometryCollectionComponent*> GeometryCollectionComponents;
+		Actor->GetComponents(GeometryCollectionComponents);
+
+		for (UGeometryCollectionComponent* GeometryCollectionComponent : GeometryCollectionComponents)
+		{
+			FScopedColorEdit EditBoneColor = GeometryCollectionComponent->EditBoneSelection();
+			EditBoneColor.SelectBones(SelectionMode);
+			InToolkit->SetBoneSelection(GeometryCollectionComponent, EditBoneColor.GetSelectedBones(), true);
+		}
+	}
+}
+
+
+
+FText UFractureToolSelectNone::GetDisplayText() const
+{
+	return FText(NSLOCTEXT("Fracture", "FractureToolSelectNone", "Select None"));
+}
+
+FText UFractureToolSelectNone::GetTooltipText() const
+{
+	return FText(NSLOCTEXT("Fracture", "FractureToolSelectNoneTooltip", "Deselects all Bones in the GeometryCollection."));
+}
+
+FSlateIcon UFractureToolSelectNone::GetToolIcon() const
+{
+	return FSlateIcon("FractureEditorStyle", "FractureEditor.SelectNone");
+}
+
+void UFractureToolSelectNone::RegisterUICommand(FFractureEditorCommands* BindingContext)
+{
+	UI_COMMAND_EXT(BindingContext, UICommandInfo, "SelectNone", "Select None", "Deselects all Bones in the GeometryCollection.", EUserInterfaceActionType::Button, FInputChord(EModifierKey::Control | EModifierKey::Shift, EKeys::D));
+	BindingContext->SelectNone = UICommandInfo;
+}
+
+void UFractureToolSelectNone::Execute(TWeakPtr<FFractureEditorModeToolkit> InToolkit)
+{
+	if (InToolkit.IsValid())
+	{
+		SelectByMode(InToolkit.Pin().Get(), GeometryCollection::ESelectionMode::None);
+	}
+}
+
+
+FText UFractureToolSelectNeighbors::GetDisplayText() const
+{
+	return FText(NSLOCTEXT("Fracture", "FractureToolSelectNeighbors", "Select Neighbors"));
+}
+
+FText UFractureToolSelectNeighbors::GetTooltipText() const
+{
+	return FText(NSLOCTEXT("Fracture", "FractureToolSelectNeighborsTooltip", "Select all bones adjacent to the currently selected bones."));
+}
+
+FSlateIcon UFractureToolSelectNeighbors::GetToolIcon() const
+{
+	return FSlateIcon("FractureEditorStyle", "FractureEditor.SelectNeighbors");
+}
+
+void UFractureToolSelectNeighbors::RegisterUICommand(FFractureEditorCommands* BindingContext)
+{
+	UI_COMMAND_EXT(BindingContext, UICommandInfo, "SelectNeighbors", "Select Neighbors", "Select all bones adjacent to the currently selected bones.", EUserInterfaceActionType::Button, FInputChord());
+	BindingContext->SelectNeighbors = UICommandInfo;
+}
+
+void UFractureToolSelectNeighbors::Execute(TWeakPtr<FFractureEditorModeToolkit> InToolkit)
+{
+	if (InToolkit.IsValid())
+	{
+		SelectByMode(InToolkit.Pin().Get(), GeometryCollection::ESelectionMode::Neighbors);
+	}
+}
+
+
+FText UFractureToolSelectSiblings::GetDisplayText() const
+{
+	return FText(NSLOCTEXT("Fracture", "FractureToolSelectSiblings", "Select Siblings"));
+}
+
+FText UFractureToolSelectSiblings::GetTooltipText() const
+{
+	return FText(NSLOCTEXT("Fracture", "FractureToolSelectSiblingsTooltip", "Select all bones at the same levels as the currently selected bones."));
+}
+
+FSlateIcon UFractureToolSelectSiblings::GetToolIcon() const
+{
+	return FSlateIcon("FractureEditorStyle", "FractureEditor.SelectSiblings");
+}
+
+void UFractureToolSelectSiblings::RegisterUICommand(FFractureEditorCommands* BindingContext)
+{
+	UI_COMMAND_EXT(BindingContext, UICommandInfo, "SelectSiblings", "Select Siblings", "Select all bones at the same levels as the currently selected bones.", EUserInterfaceActionType::Button, FInputChord());
+	BindingContext->SelectSiblings = UICommandInfo;
+}
+
+void UFractureToolSelectSiblings::Execute(TWeakPtr<FFractureEditorModeToolkit> InToolkit)
+{
+	if (InToolkit.IsValid())
+	{
+		SelectByMode(InToolkit.Pin().Get(), GeometryCollection::ESelectionMode::Siblings);
+	}
+}
+
+
+FText UFractureToolSelectAllInCluster::GetDisplayText() const
+{
+	return FText(NSLOCTEXT("Fracture", "FractureToolSelectAllInCluster", "Select All In Cluster"));
+}
+
+FText UFractureToolSelectAllInCluster::GetTooltipText() const
+{
+	return FText(NSLOCTEXT("Fracture", "FractureToolSelectAllInClusterTooltip", "Select all bones with the same parent as selected bones."));
+}
+
+FSlateIcon UFractureToolSelectAllInCluster::GetToolIcon() const
+{
+	return FSlateIcon("FractureEditorStyle", "FractureEditor.SelectAllInCluster");
+}
+
+void UFractureToolSelectAllInCluster::RegisterUICommand(FFractureEditorCommands* BindingContext)
+{
+	UI_COMMAND_EXT(BindingContext, UICommandInfo, "SelectAllInCluster", "Select All In Cluster", "Select all bones with the same parent as selected bones.", EUserInterfaceActionType::Button, FInputChord());
+	BindingContext->SelectAllInCluster = UICommandInfo;
+}
+
+void UFractureToolSelectAllInCluster::Execute(TWeakPtr<FFractureEditorModeToolkit> InToolkit)
+{
+	if (InToolkit.IsValid())
+	{
+		SelectByMode(InToolkit.Pin().Get(), GeometryCollection::ESelectionMode::AllInCluster);
+	}
+}
+
+
+FText UFractureToolSelectInvert::GetDisplayText() const
+{
+	return FText(NSLOCTEXT("Fracture", "FractureToolSelectInvert", "Invert Selection"));
+}
+
+FText UFractureToolSelectInvert::GetTooltipText() const
+{
+	return FText(NSLOCTEXT("Fracture", "FractureToolSelectInvertTooltip", "Invert current selection of bones."));
+}
+
+FSlateIcon UFractureToolSelectInvert::GetToolIcon() const
+{
+	return FSlateIcon("FractureEditorStyle", "FractureEditor.SelectInvert");
+}
+
+void UFractureToolSelectInvert::RegisterUICommand(FFractureEditorCommands* BindingContext)
+{
+	UI_COMMAND_EXT(BindingContext, UICommandInfo, "SelectInvert", "Invert Selection", "Invert current selection of bones.", EUserInterfaceActionType::Button, FInputChord());
+	BindingContext->SelectInvert = UICommandInfo;
+}
+
+void UFractureToolSelectInvert::Execute(TWeakPtr<FFractureEditorModeToolkit> InToolkit)
+{
+	if (InToolkit.IsValid())
+	{
+		SelectByMode(InToolkit.Pin().Get(), GeometryCollection::ESelectionMode::InverseGeometry);
+	}
+}
\ No newline at end of file
diff --git a/Engine/Plugins/Experimental/ChaosEditor/Source/FractureEditor/Private/FractureToolSelectors.h b/Engine/Plugins/Experimental/ChaosEditor/Source/FractureEditor/Private/FractureToolSelectors.h
new file mode 100644
--- /dev/null
+++ b/Engine/Plugins/Experimental/ChaosEditor/Source/FractureEditor/Private/FractureToolSelectors.h
@@ -0,0 +1,120 @@
+// Copyright Epic Games, Inc. All Rights Reserved.
+
+#pragma once
+
+#include "FractureTool.h"
+
+#include "FractureToolSelectors.generated.h"
+
+
+namespace GeometryCollection
+{
+	enum class ESelectionMode : uint8;
+}
+
+class FFractureEditorModeToolkit;
+
+
+UCLASS(DisplayName = "Select All", Category = "FractureTools")
+class UFractureToolSelectAll : public UFractureActionTool
+{
+public:
+	GENERATED_BODY()
+
+	UFractureToolSelectAll(const FObjectInitializer& ObjInit) : Super(ObjInit) {}
+
+	// UFractureActionTool Interface
+	virtual FText GetDisplayText() const override;
+	virtual FText GetTooltipText() const override;
+	virtual FSlateIcon GetToolIcon() const override;
+	virtual void RegisterUICommand(FFractureEditorCommands* BindingContext) override;
+	virtual void Execute(TWeakPtr<FFractureEditorModeToolkit> InToolkit) override;
+
+protected:
+	void SelectByMode(FFractureEditorModeToolkit* InToolkit, GeometryCollection::ESelectionMode SelectionMode);
+};
+
+
+UCLASS(DisplayName = "Select None", Category = "FractureTools")
+class UFractureToolSelectNone : public UFractureToolSelectAll
+{
+public:
+	GENERATED_BODY()
+
+	UFractureToolSelectNone(const FObjectInitializer& ObjInit) : Super(ObjInit) {}
+
+	// UFractureActionTool Interface
+	virtual FText GetDisplayText() const override;
+	virtual FText GetTooltipText() const override;
+	virtual FSlateIcon GetToolIcon() const override;
+	virtual void RegisterUICommand(FFractureEditorCommands* BindingContext) override;
+	virtual void Execute(TWeakPtr<FFractureEditorModeToolkit> InToolkit) override;
+};
+
+
+UCLASS(DisplayName = "Select Neighbors", Category = "FractureTools")
+class UFractureToolSelectNeighbors : public UFractureToolSelectAll
+{
+public:
+	GENERATED_BODY()
+
+	UFractureToolSelectNeighbors(const FObjectInitializer& ObjInit) : Super(ObjInit) {}
+
+	// UFractureActionTool Interface
+	virtual FText GetDisplayText() const override;
+	virtual FText GetTooltipText() const override;
+	virtual FSlateIcon GetToolIcon() const override;
+	virtual void RegisterUICommand(FFractureEditorCommands* BindingContext) override;
+	virtual void Execute(TWeakPtr<FFractureEditorModeToolkit> InToolkit) override;
+};
+
+
+UCLASS(DisplayName = "Select Siblings", Category = "FractureTools")
+class UFractureToolSelectSiblings : public UFractureToolSelectAll
+{
+public:
+	GENERATED_BODY()
+
+	UFractureToolSelectSiblings(const FObjectInitializer& ObjInit) : Super(ObjInit) {}
+
+	// UFractureActionTool Interface
+	virtual FText GetDisplayText() const override;
+	virtual FText GetTooltipText() const override;
+	virtual FSlateIcon GetToolIcon() const override;
+	virtual void RegisterUICommand(FFractureEditorCommands* BindingContext) override;
+	virtual void Execute(TWeakPtr<FFractureEditorModeToolkit> InToolkit) override;
+};
+
+
+UCLASS(DisplayName = "Select Neighbors", Category = "FractureTools")
+class UFractureToolSelectAllInCluster : public UFractureToolSelectAll
+{
+public:
+	GENERATED_BODY()
+
+	UFractureToolSelectAllInCluster(const FObjectInitializer& ObjInit) : Super(ObjInit) {}
+
+	// UFractureActionTool Interface
+	virtual FText GetDisplayText() const override;
+	virtual FText GetTooltipText() const override;
+	virtual FSlateIcon GetToolIcon() const override;
+	virtual void RegisterUICommand(FFractureEditorCommands* BindingContext) override;
+	virtual void Execute(TWeakPtr<FFractureEditorModeToolkit> InToolkit) override;
+};
+
+
+UCLASS(DisplayName = "Select Neighbors", Category = "FractureTools")
+class UFractureToolSelectInvert : public UFractureToolSelectAll
+{
+public:
+	GENERATED_BODY()
+
+	UFractureToolSelectInvert(const FObjectInitializer& ObjInit) : Super(ObjInit) {}
+
+	// UFractureActionTool Interface
+	virtual FText GetDisplayText() const override;
+	virtual FText GetTooltipText() const override;
+	virtual FSlateIcon GetToolIcon() const override;
+	virtual void RegisterUICommand(FFractureEditorCommands* BindingContext) override;
+	virtual void Execute(TWeakPtr<FFractureEditorModeToolkit> InToolkit) override;
+};
diff --git a/Engine/Plugins/Experimental/ChaosEditor/Source/FractureEditor/Private/FractureToolSlice.cpp b/Engine/Plugins/Experimental/ChaosEditor/Source/FractureEditor/Private/FractureToolSlice.cpp
--- a/Engine/Plugins/Experimental/ChaosEditor/Source/FractureEditor/Private/FractureToolSlice.cpp
+++ b/Engine/Plugins/Experimental/ChaosEditor/Source/FractureEditor/Private/FractureToolSlice.cpp
@@ -13,23 +13,6 @@
 
 #define LOCTEXT_NAMESPACE "FractureSlice"
 
-void UFractureSliceSettings::PostEditChangeProperty(FPropertyChangedEvent& PropertyChangedEvent)
-{
-	if (OwnerTool != nullptr)
-	{
-		OwnerTool->PostEditChangeProperty(PropertyChangedEvent);
-	}
-	Super::PostEditChangeProperty(PropertyChangedEvent);
-}
-
-void UFractureSliceSettings::PostEditChangeChainProperty(struct FPropertyChangedChainEvent& PropertyChangedEvent)
-{
-	if (OwnerTool != nullptr)
-	{
-		OwnerTool->PostEditChangeChainProperty(PropertyChangedEvent);
-	}
-	Super::PostEditChangeChainProperty(PropertyChangedEvent);
-}
 
 
 UFractureToolSlice::UFractureToolSlice(const FObjectInitializer& ObjInit) 
@@ -63,15 +46,13 @@ void UFractureToolSlice::RegisterUICommand( FFractureEditorCommands* BindingCont
 TArray<UObject*> UFractureToolSlice::GetSettingsObjects() const 
 { 
 	TArray<UObject*> Settings; 
-	Settings.Add(GetMutableDefault<UFractureCommonSettings>());
-	Settings.Add(GetMutableDefault<UFractureSliceSettings>());
+	Settings.Add(CutterSettings);
+	Settings.Add(SliceSettings);
 	return Settings;
 }
 
-void UFractureToolSlice::GenerateSliceTransforms(const FFractureContext& Context, TArray<FTransform>& CuttingPlaneTransforms)
+void UFractureToolSlice::GenerateSliceTransforms(const FFractureToolContext& Context, TArray<FTransform>& CuttingPlaneTransforms)
 {
-	const UFractureSliceSettings* LocalSliceSettings = GetMutableDefault<UFractureSliceSettings>();
-
 	const FBox& Bounds = Context.Bounds;
 	const FVector& Min = Context.Bounds.Min;
 	const FVector& Max = Context.Bounds.Max;
@@ -79,21 +60,20 @@ void UFractureToolSlice::GenerateSliceTransforms(const FFractureContext& Context
 	const FVector Extents(Context.Bounds.Max - Context.Bounds.Min);
 	const FVector HalfExtents(Extents * 0.5f);
 
-	const FVector Step(Extents.X / (LocalSliceSettings->SlicesX+1), Extents.Y / (LocalSliceSettings->SlicesY+1), Extents.Z / (LocalSliceSettings->SlicesZ+1));
+	const FVector Step(Extents.X / (SliceSettings->SlicesX+1), Extents.Y / (SliceSettings->SlicesY+1), Extents.Z / (SliceSettings->SlicesZ+1));
 
-	CuttingPlaneTransforms.Reserve(LocalSliceSettings->SlicesX * LocalSliceSettings->SlicesY * LocalSliceSettings->SlicesZ);
+	CuttingPlaneTransforms.Reserve(SliceSettings->SlicesX * SliceSettings->SlicesY * SliceSettings->SlicesZ);
 
 	FRandomStream RandomStream(Context.RandomSeed);
 
-	const float SliceAngleVariationInRadians = FMath::DegreesToRadians(LocalSliceSettings->SliceAngleVariation);
-
+	const float SliceAngleVariationInRadians = FMath::DegreesToRadians(SliceSettings->SliceAngleVariation);
 
 	const FVector XMin(-Bounds.GetExtent() * FVector(1.0f, 1.0f, 0.0f));
 	const FVector XMax(Bounds.GetExtent() * FVector(1.0f, 1.0f, 0.0f));
 
-	for (int32 xx = 0; xx < LocalSliceSettings->SlicesX; ++xx)
+	for (int32 xx = 0; xx < SliceSettings->SlicesX; ++xx)
 	{
-		const FVector SlicePosition(FVector(Min.X, Center.Y, Center.Z) + FVector((Step.X * xx) + Step.X, 0.0f, 0.0f) + RandomStream.VRand() * RandomStream.GetFraction() * LocalSliceSettings->SliceOffsetVariation);
+		const FVector SlicePosition(FVector(Min.X, Center.Y, Center.Z) + FVector((Step.X * xx) + Step.X, 0.0f, 0.0f) + RandomStream.VRand() * RandomStream.GetFraction() * SliceSettings->SliceOffsetVariation);
 		FTransform Transform(FQuat(FVector::RightVector, FMath::DegreesToRadians(90)), SlicePosition);
 		const FQuat RotA(FVector::RightVector, RandomStream.FRandRange(0.0f, SliceAngleVariationInRadians));
 		const FQuat RotB(FVector::ForwardVector, RandomStream.FRandRange(0.0f, SliceAngleVariationInRadians));
@@ -101,14 +81,12 @@ void UFractureToolSlice::GenerateSliceTransforms(const FFractureContext& Context
 		CuttingPlaneTransforms.Emplace(Transform);
 	}
 
-
-
 	const FVector YMin(-Bounds.GetExtent() * FVector(1.0f, 1.0f, 0.0f));
 	const FVector YMax(Bounds.GetExtent() * FVector(1.0f, 1.0f, 0.0f));
 
-	for (int32 yy = 0; yy < LocalSliceSettings->SlicesY; ++yy)
+	for (int32 yy = 0; yy < SliceSettings->SlicesY; ++yy)
 	{
-		const FVector SlicePosition(FVector(Center.X, Min.Y, Center.Z) + FVector(0.0f, (Step.Y * yy) + Step.Y, 0.0f) + RandomStream.VRand() * RandomStream.GetFraction() * LocalSliceSettings->SliceOffsetVariation);
+		const FVector SlicePosition(FVector(Center.X, Min.Y, Center.Z) + FVector(0.0f, (Step.Y * yy) + Step.Y, 0.0f) + RandomStream.VRand() * RandomStream.GetFraction() * SliceSettings->SliceOffsetVariation);
 		FTransform Transform(FQuat(FVector::ForwardVector, FMath::DegreesToRadians(90)), SlicePosition);
 		const FQuat RotA(FVector::RightVector, RandomStream.FRandRange(0.0f, SliceAngleVariationInRadians));
 		const FQuat RotB(FVector::ForwardVector, RandomStream.FRandRange(0.0f, SliceAngleVariationInRadians));
@@ -116,15 +94,12 @@ void UFractureToolSlice::GenerateSliceTransforms(const FFractureContext& Context
 		CuttingPlaneTransforms.Emplace(Transform);
 	}
 
-
-
-
 	const FVector ZMin(-Bounds.GetExtent() * FVector(1.0f, 1.0f, 0.0f));
 	const FVector ZMax(Bounds.GetExtent() * FVector(1.0f, 1.0f, 0.0f));
 
-	for (int32 zz = 0; zz < LocalSliceSettings->SlicesZ; ++zz)
+	for (int32 zz = 0; zz < SliceSettings->SlicesZ; ++zz)
 	{
-		const FVector SlicePosition(FVector(Center.X, Center.Y, Min.Z) + FVector(0.0f, 0.0f, (Step.Z * zz) + Step.Z) + RandomStream.VRand() * RandomStream.GetFraction() * LocalSliceSettings->SliceOffsetVariation);
+		const FVector SlicePosition(FVector(Center.X, Center.Y, Min.Z) + FVector(0.0f, 0.0f, (Step.Z * zz) + Step.Z) + RandomStream.VRand() * RandomStream.GetFraction() * SliceSettings->SliceOffsetVariation);
 		FTransform Transform(SlicePosition);
 		const FQuat RotA(FVector::RightVector, RandomStream.FRandRange(0.0f, SliceAngleVariationInRadians));
 		const FQuat RotB(FVector::ForwardVector, RandomStream.FRandRange(0.0f, SliceAngleVariationInRadians));
@@ -133,20 +108,8 @@ void UFractureToolSlice::GenerateSliceTransforms(const FFractureContext& Context
 	}
 }
 
-
-#if WITH_EDITOR
-void UFractureToolSlice::PostEditChangeChainProperty(struct FPropertyChangedChainEvent& PropertyChangedEvent)
-{
-	FractureContextChanged();
-}
-#endif
-
-
-
 void UFractureToolSlice::Render(const FSceneView* View, FViewport* Viewport, FPrimitiveDrawInterface* PDI)
 {
-	const UFractureCommonSettings* LocalCommonSettings = GetDefault<UFractureCommonSettings>();
-
 	for (const FTransform& Transform : RenderCuttingPlanesTransforms)
 	{
 		PDI->DrawPoint(Transform.GetLocation(), FLinearColor::Green, 4.f, SDPG_Foreground);
@@ -161,15 +124,13 @@ void UFractureToolSlice::Render(const FSceneView* View, FViewport* Viewport, FPr
 
 void UFractureToolSlice::FractureContextChanged()
 {
-	const UFractureCommonSettings* LocalCommonSettings = GetDefault<UFractureCommonSettings>();
-
-	TArray<FFractureContext> FractureContexts;
-	FFractureEditorModeToolkit::GetFractureContexts(FractureContexts);
+	TArray<FFractureToolContext> FractureContexts;
+	GetFractureContexts(FractureContexts);
 
 	RenderCuttingPlanesTransforms.Empty();
 
 	RenderCuttingPlaneSize = FLT_MAX;
-	for (FFractureContext& FractureContext : FractureContexts)
+	for (FFractureToolContext& FractureContext : FractureContexts)
 	{
 		// Move the local bounds to the actor so we weill draw in the correct location
 		FractureContext.Bounds = FractureContext.Bounds.TransformBy(FractureContext.Transform);
@@ -182,12 +143,8 @@ void UFractureToolSlice::FractureContextChanged()
 	}
 }
 
-
-void UFractureToolSlice::ExecuteFracture(const FFractureContext& FractureContext)
+void UFractureToolSlice::ExecuteFracture(const FFractureToolContext& FractureContext)
 {
-	const UFractureCommonSettings* LocalCommonSettings = GetDefault<UFractureCommonSettings>();
-	const UFractureSliceSettings* LocalSliceSettings = GetMutableDefault<UFractureSliceSettings>();
-
 	if (FractureContext.FracturedGeometryCollection != nullptr)
 	{
 		TSharedPtr<FGeometryCollection, ESPMode::ThreadSafe> GeometryCollectionPtr = FractureContext.FracturedGeometryCollection->GetGeometryCollection();
@@ -206,12 +163,12 @@ void UFractureToolSlice::ExecuteFracture(const FFractureContext& FractureContext
 
 			FInternalSurfaceMaterials InternalSurfaceMaterials;
 			FNoiseSettings NoiseSettings;
-			if (LocalCommonSettings->Amplitude > 0.0f)
+			if (CutterSettings->Amplitude > 0.0f)
 			{
-				NoiseSettings.Amplitude = LocalCommonSettings->Amplitude;
-				NoiseSettings.Frequency = LocalCommonSettings->Frequency;
-				NoiseSettings.Octaves = LocalCommonSettings->OctaveNumber;
-				NoiseSettings.PointSpacing = LocalCommonSettings->SurfaceResolution;
+				NoiseSettings.Amplitude = CutterSettings->Amplitude;
+				NoiseSettings.Frequency = CutterSettings->Frequency;
+				NoiseSettings.Octaves = CutterSettings->OctaveNumber;
+				NoiseSettings.PointSpacing = CutterSettings->SurfaceResolution;
 				InternalSurfaceMaterials.NoiseSettings = NoiseSettings;
 			}
 
@@ -220,9 +177,4 @@ void UFractureToolSlice::ExecuteFracture(const FFractureContext& FractureContext
 	}
 }
 
-bool UFractureToolSlice::CanExecuteFracture() const
-{
-	return FFractureEditorModeToolkit::IsLeafBoneSelected();
-}
-
 #undef LOCTEXT_NAMESPACE
\ No newline at end of file
diff --git a/Engine/Plugins/Experimental/ChaosEditor/Source/FractureEditor/Private/FractureToolSlice.h b/Engine/Plugins/Experimental/ChaosEditor/Source/FractureEditor/Private/FractureToolSlice.h
--- a/Engine/Plugins/Experimental/ChaosEditor/Source/FractureEditor/Private/FractureToolSlice.h
+++ b/Engine/Plugins/Experimental/ChaosEditor/Source/FractureEditor/Private/FractureToolSlice.h
@@ -2,20 +2,16 @@
 
 #pragma once
 
-#include "FractureTool.h"
+#include "FractureToolCutter.h"
 
 #include "FractureToolSlice.generated.h"
 
 
-UCLASS()
-class UFractureSliceSettings : public UObject
+UCLASS(config = EditorPerProjectUserSettings)
+class UFractureSliceSettings : public UFractureToolSettings
 {
-	GENERATED_BODY()
 public:
-#if WITH_EDITOR
-	virtual void PostEditChangeProperty(FPropertyChangedEvent& PropertyChangedEvent) override;
-	virtual void PostEditChangeChainProperty(struct FPropertyChangedChainEvent& PropertyChangedEvent) override;
-#endif
+	GENERATED_BODY()
 
 	UFractureSliceSettings(const FObjectInitializer& ObjInit)
 		: Super(ObjInit)
@@ -45,47 +41,36 @@ public:
 	/** Slicing Offset Variation - Slicing Method [0..1] */
 	UPROPERTY(EditAnywhere, Category = Slicing, meta = (DisplayName = "Random Offset Variation", UIMin = "0.0"))
 	float SliceOffsetVariation;
-
-	UPROPERTY()
-	UFractureTool *OwnerTool;
-
 };
 
 UCLASS(DisplayName="Slice Tool", Category="FractureTools")
-class UFractureToolSlice : public UFractureTool
+class UFractureToolSlice : public UFractureToolCutterBase
 {
 public:
 	GENERATED_BODY()
 
-	UFractureToolSlice(const FObjectInitializer& ObjInit);//  : Super(ObjInit) {}
+	UFractureToolSlice(const FObjectInitializer& ObjInit);
 
 	// UFractureTool Interface
 	virtual FText GetDisplayText() const override;
 	virtual FText GetTooltipText() const override;
 	virtual FSlateIcon GetToolIcon() const override;
 
-	virtual void RegisterUICommand( FFractureEditorCommands* BindingContext );
+	virtual void RegisterUICommand( FFractureEditorCommands* BindingContext ) override;
 	
-	virtual TArray<UObject*> GetSettingsObjects() const;// { return TArray<UObject*>(); }
+	virtual TArray<UObject*> GetSettingsObjects() const override;
 	void Render(const FSceneView* View, FViewport* Viewport, FPrimitiveDrawInterface* PDI) override;
 
 	virtual void FractureContextChanged() override;
-	virtual void ExecuteFracture(const FFractureContext& Context) override;
-	virtual bool CanExecuteFracture() const override;
-
-	void GenerateSliceTransforms(const FFractureContext& Context, TArray<FTransform>& CuttingPlaneTransforms);
-
-#if WITH_EDITOR
-	virtual void PostEditChangeChainProperty(struct FPropertyChangedChainEvent& PropertyChangedEvent) override;
-#endif
+	virtual void ExecuteFracture(const FFractureToolContext& Context) override;
 
+	void GenerateSliceTransforms(const FFractureToolContext& Context, TArray<FTransform>& CuttingPlaneTransforms);
 
 	// Slicing
 	UPROPERTY(EditAnywhere, Category = Slicing)
 	UFractureSliceSettings* SliceSettings;
-private:
-// 	TArray<TTuple<FVector, FVector>> Edges;
 
+private:
 	float RenderCuttingPlaneSize;
 	TArray<FTransform> RenderCuttingPlanesTransforms;
 };
\ No newline at end of file
diff --git a/Engine/Plugins/Experimental/ChaosEditor/Source/FractureEditor/Private/FractureToolUniform.cpp b/Engine/Plugins/Experimental/ChaosEditor/Source/FractureEditor/Private/FractureToolUniform.cpp
--- a/Engine/Plugins/Experimental/ChaosEditor/Source/FractureEditor/Private/FractureToolUniform.cpp
+++ b/Engine/Plugins/Experimental/ChaosEditor/Source/FractureEditor/Private/FractureToolUniform.cpp
@@ -7,30 +7,12 @@
 
 #define LOCTEXT_NAMESPACE "FractureUniform"
 
-void UFractureUniformSettings::PostEditChangeProperty(FPropertyChangedEvent& PropertyChangedEvent)
-{
-	if (OwnerTool != nullptr)
-	{
-		OwnerTool->PostEditChangeProperty(PropertyChangedEvent);
-	}
-	Super::PostEditChangeProperty(PropertyChangedEvent);
-}
-
-void UFractureUniformSettings::PostEditChangeChainProperty(struct FPropertyChangedChainEvent& PropertyChangedEvent)
-{
-	if (OwnerTool != nullptr)
-	{
-		OwnerTool->PostEditChangeChainProperty(PropertyChangedEvent);
-	}
-	Super::PostEditChangeChainProperty(PropertyChangedEvent);
-}
-
 
 UFractureToolUniform::UFractureToolUniform(const FObjectInitializer& ObjInit) 
 	: Super(ObjInit) 
 {
-	Settings = NewObject<UFractureUniformSettings>(GetTransientPackage(), UFractureUniformSettings::StaticClass());
-	Settings->OwnerTool = this;
+	UniformSettings = NewObject<UFractureUniformSettings>(GetTransientPackage(), UFractureUniformSettings::StaticClass());
+	UniformSettings->OwnerTool = this;
 }
 
 FText UFractureToolUniform::GetDisplayText() const
@@ -57,18 +39,18 @@ void UFractureToolUniform::RegisterUICommand( FFractureEditorCommands* BindingCo
 TArray<UObject*> UFractureToolUniform::GetSettingsObjects() const 
 { 
 	TArray<UObject*> AllSettings; 
-	AllSettings.Add(GetMutableDefault<UFractureCommonSettings>());
-	AllSettings.Add(Settings);
+	AllSettings.Add(CutterSettings);
+	AllSettings.Add(UniformSettings);
 	return AllSettings;
 }
 
-void UFractureToolUniform::GenerateVoronoiSites(const FFractureContext &Context, TArray<FVector>& Sites)
+void UFractureToolUniform::GenerateVoronoiSites(const FFractureToolContext &Context, TArray<FVector>& Sites)
 {
 	FRandomStream RandStream(Context.RandomSeed);
 
 	const FVector Extent(Context.Bounds.Max - Context.Bounds.Min);
 
-	const int32 SiteCount = RandStream.RandRange(Settings->NumberVoronoiSitesMin, Settings->NumberVoronoiSitesMax);
+	const int32 SiteCount = RandStream.RandRange(UniformSettings->NumberVoronoiSitesMin, UniformSettings->NumberVoronoiSitesMax);
 
 	Sites.Reserve(Sites.Num() + SiteCount);
 	for (int32 ii = 0; ii < SiteCount; ++ii)
@@ -77,4 +59,7 @@ void UFractureToolUniform::GenerateVoronoiSites(const FFractureContext &Context,
 	}
 }
 
+
+
+
 #undef LOCTEXT_NAMESPACE
diff --git a/Engine/Plugins/Experimental/ChaosEditor/Source/FractureEditor/Private/FractureToolUniform.h b/Engine/Plugins/Experimental/ChaosEditor/Source/FractureEditor/Private/FractureToolUniform.h
--- a/Engine/Plugins/Experimental/ChaosEditor/Source/FractureEditor/Private/FractureToolUniform.h
+++ b/Engine/Plugins/Experimental/ChaosEditor/Source/FractureEditor/Private/FractureToolUniform.h
@@ -2,30 +2,24 @@
 
 #pragma once
 
-#include "FractureToolVoronoiBase.h"
+#include "FractureToolCutter.h"
 
 #include "FractureToolUniform.generated.h"
 
 
-UCLASS()
-class UFractureUniformSettings
-	: public UObject
-{
-	GENERATED_BODY()
 
+UCLASS(config = EditorPerProjectUserSettings)
+class UFractureUniformSettings : public UFractureToolSettings
+{
 public:
-#if WITH_EDITOR
-	virtual void PostEditChangeProperty(FPropertyChangedEvent& PropertyChangedEvent) override;
-	virtual void PostEditChangeChainProperty(struct FPropertyChangedChainEvent& PropertyChangedEvent) override;
-#endif
+	GENERATED_BODY()
 
-	UFractureUniformSettings()
-		: NumberVoronoiSitesMin(20)
+	UFractureUniformSettings(const FObjectInitializer& ObjInit)
+		: Super(ObjInit)
+		, NumberVoronoiSitesMin(20)
 		, NumberVoronoiSitesMax(20)
-		, OwnerTool(nullptr)
 	{}
 
-
 	/** Minimum Number of Voronoi sites - A random number will be chosen between the Min and Max for each bone you have selected */
 	UPROPERTY(EditAnywhere, Category = UniformVoronoi, meta = (DisplayName = "Minimum Voronoi Sites", UIMin = "1", UIMax = "5000", ClampMin = "1"))
 	int32 NumberVoronoiSitesMin;
@@ -34,38 +28,30 @@ public:
 	UPROPERTY(EditAnywhere, Category = UniformVoronoi, meta = (DisplayName = "Maximum Voronoi Sites", UIMin = "1", UIMax = "5000", ClampMin = "1"))
 	int32 NumberVoronoiSitesMax;
 
-	UPROPERTY()
-	UFractureTool *OwnerTool;
 };
 
 
 UCLASS(DisplayName="Uniform Voronoi", Category="FractureTools")
-class UFractureToolUniform : public UFractureToolVoronoiBase
+class UFractureToolUniform : public UFractureToolVoronoiCutterBase
 {
 public:
 	GENERATED_BODY()
 
-	UFractureToolUniform(const FObjectInitializer& ObjInit);//  : Super(ObjInit) {}
+	UFractureToolUniform(const FObjectInitializer& ObjInit);
 
 	// UFractureTool Interface
 	virtual FText GetDisplayText() const override;
 	virtual FText GetTooltipText() const override;
 	virtual FSlateIcon GetToolIcon() const override;
-	virtual TArray<UObject*> GetSettingsObjects() const override;// { return TArray<UObject*>(); }
-
-	virtual void RegisterUICommand( FFractureEditorCommands* BindingContext );
+	virtual TArray<UObject*> GetSettingsObjects() const override;
+	virtual void RegisterUICommand(FFractureEditorCommands* BindingContext) override;
 
-// #if WITH_EDITOR
-// 	virtual void PostEditChangeProperty(FPropertyChangedEvent& PropertyChangedEvent) override;
-// 	virtual void PostEditChangeChainProperty(struct FPropertyChangedChainEvent& PropertyChangedEvent) override;
-// #endif
-// 	void Render(const FSceneView* View, FViewport* Viewport, FPrimitiveDrawInterface* PDI) override;
 
 	// Uniform Voronoi Fracture Input Settings
 	UPROPERTY(EditAnywhere, Category = Uniform)
-	UFractureUniformSettings* Settings;
+	UFractureUniformSettings* UniformSettings;
 
 protected:
+	void GenerateVoronoiSites(const FFractureToolContext&Context, TArray<FVector>& Sites) override;
+};
 
-	void GenerateVoronoiSites(const FFractureContext &Context, TArray<FVector>& Sites) override;
-};
\ No newline at end of file
diff --git a/Engine/Plugins/Experimental/ChaosEditor/Source/FractureEditor/Private/FractureToolVoronoiBase.cpp b/Engine/Plugins/Experimental/ChaosEditor/Source/FractureEditor/Private/FractureToolVoronoiBase.cpp
deleted file mode 100644
diff --git a/Engine/Plugins/Experimental/ChaosEditor/Source/FractureEditor/Private/FractureToolVoronoiBase.h b/Engine/Plugins/Experimental/ChaosEditor/Source/FractureEditor/Private/FractureToolVoronoiBase.h
deleted file mode 100644
diff --git a/Engine/Plugins/Experimental/ChaosEditor/Source/FractureEditor/Public/FractureEditorModeToolkit.h b/Engine/Plugins/Experimental/ChaosEditor/Source/FractureEditor/Public/FractureEditorModeToolkit.h
--- a/Engine/Plugins/Experimental/ChaosEditor/Source/FractureEditor/Public/FractureEditorModeToolkit.h
+++ b/Engine/Plugins/Experimental/ChaosEditor/Source/FractureEditor/Public/FractureEditorModeToolkit.h
@@ -13,13 +13,15 @@ class IDetailsView;
 class IPropertyHandle;
 class SScrollBox;
 class UFractureTool;
-struct FFractureContext;
+struct FFractureToolContext;
 class FGeometryCollection;
 class SGeometryCollectionOutliner;
 class AGeometryCollectionActor;
 class UGeometryCollectionComponent;
 class UGeometryCollection;
 class FFractureEditorModeToolkit;
+class UFractureActionTool;
+class UFractureModalTool;
 
 namespace GeometryCollection
 {
@@ -57,16 +59,16 @@ public:
 	/** FGCObject interface */
 	virtual void AddReferencedObjects(FReferenceCollector& Collector) override;
 
-	void SetActiveTool(UFractureTool* InActiveTool);
-	UFractureTool* GetActiveTool() const;
-	bool IsActiveTool(UFractureTool* InActiveTool);
+	void ExecuteAction(UFractureActionTool* InActionTool);
+	bool CanExecuteAction(UFractureActionTool* InActionTool) const;
+
+	void SetActiveTool(UFractureModalTool* InActiveTool);
+	UFractureModalTool* GetActiveTool() const;
+	bool IsActiveTool(UFractureModalTool* InActiveTool);
 
 	void SetOutlinerComponents(const TArray<UGeometryCollectionComponent*>& InNewComponents);
 	void SetBoneSelection(UGeometryCollectionComponent* InRootComponent, const TArray<int32>& InSelectedBones, bool bClearCurrentSelection);
 
-	// Selection Command Callbacks;
-	void OnSelectByMode(GeometryCollection::ESelectionMode);
-
 	// View Settings
 	float GetExplodedViewValue() const;
 	int32 GetLevelViewValue() const;
@@ -85,41 +87,15 @@ public:
 	TSharedRef<SWidget> GetLevelViewMenuContent(TSharedRef<IPropertyHandle> PropertyHandle);
 	TSharedRef<SWidget> GetViewMenuContent();
 
-
-	// Cluster Command Callbacks
-	void OnCluster();
-	void OnUncluster();
-	void OnFlatten();
-	void OnFlattenToLevel(){}
-	void OnMerge(){}
-	void OnMoveUp();
-	void GenerateAsset();
-	void ResetAsset();
-
 	void ViewUpOneLevel();
 	void ViewDownOneLevel();
 
-	// Fracture Command Callback
-	FReply OnFractureClicked();
-	bool CanExecuteFracture() const;
-
-	static bool IsLeafBoneSelected();
-
-	static ULevel* GetSelectedLevel();
- 	static AActor* AddActor(ULevel* InLevel, UClass* Class);
-	class AGeometryCollectionActor* CreateNewGeometryActor(const FString& Name, const FTransform& Transform, bool AddMaterials /*= false*/);
-
- 	void OpenGenerateAssetDialog(TArray<AActor*>& Actors);
- 	void OnGenerateAssetPathChosen(const FString& InAssetPath, TArray<AActor*> Actors);
-
-	static void GetFractureContexts(TArray<FFractureContext>& FractureContexts);
-	void ExecuteFracture(FFractureContext& FractureContext);
+	// Modal Command Callback
+	FReply OnModalClicked();
+	bool CanExecuteModal() const;
 
 	static void GetSelectedGeometryCollectionComponents(TSet<UGeometryCollectionComponent*>& GeomCompSelection);
 
-	AGeometryCollectionActor* ConvertStaticMeshToGeometryCollection(const FString& InAssetPath, TArray<AActor*>& Actors);
-
-	static void AddSingleRootNodeIfRequired(UGeometryCollection* GeometryCollectionObject);
 	static void AddAdditionalAttributesIfRequired(UGeometryCollection* GeometryCollectionObject);
 	int32 GetLevelCount();
 
@@ -129,6 +105,8 @@ public:
 	const static TArray<FName> PaletteNames;
 	virtual void GetToolPaletteNames( TArray<FName>& InPaletteName ) const { InPaletteName = PaletteNames; }
 	virtual FText GetToolPaletteDisplayName(FName PaletteName) const; 
+
+	void SetInitialPalette();
 	virtual void BuildToolPalette(FName PaletteName, class FToolBarBuilder& ToolbarBuilder);
 	virtual void OnToolPaletteChanged(FName PaletteName) override;
 
@@ -136,22 +114,22 @@ public:
 	virtual FText GetActiveToolDisplayName() const;
 	virtual FText GetActiveToolMessage() const;
 
+	void UpdateExplodedVectors(UGeometryCollectionComponent* GeometryCollectionComponent) const;
+
 	TSharedPtr<SWidget> ExplodedViewWidget;
 	TSharedPtr<SWidget> LevelViewWidget;
 	TSharedPtr<SWidget> ShowBoneColorsWidget;
 
 protected:
 	static bool IsGeometryCollectionSelected();
-	static bool IsStaticMeshSelected();
-	static bool IsSelectedActorsInEditorWorld();
-	void UpdateExplodedVectors(UGeometryCollectionComponent* GeometryCollectionComponent) const;
+	static bool IsSelectedActorsInEditorWorld();	
 
 private:
 	void OnOutlinerBoneSelectionChanged(UGeometryCollectionComponent* RootComponent, TArray<int32>& SelectedBones);
 	void BindCommands();
 
 private:
-	UFractureTool* ActiveTool;
+	UFractureModalTool* ActiveTool;
 
 	TSharedPtr<IDetailsView> DetailsView;
 	TSharedPtr<IDetailsView> ViewSettingsDetailsView;
diff --git a/Engine/Plugins/Experimental/ChaosEditor/Source/FractureEditor/Public/FractureTool.cpp b/Engine/Plugins/Experimental/ChaosEditor/Source/FractureEditor/Public/FractureTool.cpp
--- a/Engine/Plugins/Experimental/ChaosEditor/Source/FractureEditor/Public/FractureTool.cpp
+++ b/Engine/Plugins/Experimental/ChaosEditor/Source/FractureEditor/Public/FractureTool.cpp
@@ -2,22 +2,19 @@
 
 #include "FractureTool.h"
 
-#if WITH_EDITOR
-
-UFractureCommonSettings::UFractureCommonSettings(const FObjectInitializer& ObjInit)
-	: Super(ObjInit)
-	, RandomSeed(-1)
-	, ChanceToFracture(1.0)
-	, bGroupFracture(true)
-	, bDrawSites(false)
-	, bDrawDiagram(true)
-	, Amplitude(0.0f)
-	, Frequency(0.1f)
-	, OctaveNumber(4)
-	, SurfaceResolution(10)
-{}
-
-void UFractureCommonSettings::PostEditChangeProperty(FPropertyChangedEvent& PropertyChangedEvent)
+#include "Editor.h"
+#include "Engine/Selection.h"
+
+#include "GeometryCollection/GeometryCollectionObject.h"
+#include "GeometryCollection/GeometryCollection.h"
+#include "GeometryCollection/GeometryCollectionComponent.h"
+#include "GeometryCollection/GeometryCollectionClusteringUtility.h"
+#include "GeometryCollection/GeometryCollectionAlgo.h"
+
+
+DEFINE_LOG_CATEGORY(LogFractureTool);
+
+void UFractureToolSettings::PostEditChangeProperty(FPropertyChangedEvent& PropertyChangedEvent)
 {
 	if (OwnerTool != nullptr)
 	{
@@ -26,19 +23,289 @@ void UFractureCommonSettings::PostEditChangeProperty(FPropertyChangedEvent& Prop
 	Super::PostEditChangeProperty(PropertyChangedEvent);
 }
 
-void UFractureCommonSettings::PostEditChangeChainProperty(struct FPropertyChangedChainEvent& PropertyChangedEvent)
+void UFractureToolSettings::PostEditChangeChainProperty(struct FPropertyChangedChainEvent& PropertyChangedEvent)
 {
 	if (OwnerTool != nullptr)
 	{
-		OwnerTool->PostEditChangeChainProperty(PropertyChangedEvent); 
+		OwnerTool->PostEditChangeChainProperty(PropertyChangedEvent);
 	}
 	Super::PostEditChangeChainProperty(PropertyChangedEvent);
 }
 
-const TSharedPtr<FUICommandInfo>& UFractureTool::GetUICommandInfo() const
+
+
+
+const TSharedPtr<FUICommandInfo>& UFractureActionTool::GetUICommandInfo() const
 {
 	return UICommandInfo;
 }
 
-#endif
+bool UFractureActionTool::CanExecute() const
+{
+	return IsGeometryCollectionSelected();
+}
+
+bool UFractureActionTool::IsGeometryCollectionSelected()
+{
+	USelection* SelectedActors = GEditor->GetSelectedActors();
+	TArray<ULevel*> UniqueLevels;
+	for (FSelectionIterator Iter(*SelectedActors); Iter; ++Iter)
+	{
+		AActor* Actor = Cast<AActor>(*Iter);
+		if (Actor)
+		{
+			if (Actor->FindComponentByClass<UGeometryCollectionComponent>())
+			{
+				return true;
+			}
+		}
+	}
+	return false;
+}
+
+bool UFractureActionTool::IsStaticMeshSelected()
+{
+	USelection* SelectedActors = GEditor->GetSelectedActors();
+	for (FSelectionIterator Iter(*SelectedActors); Iter; ++Iter)
+	{
+		AActor* Actor = Cast<AActor>(*Iter);
+		if (Actor)
+		{
+			TInlineComponentArray<UStaticMeshComponent*> StaticMeshComponents;
+			Actor->GetComponents<UStaticMeshComponent>(StaticMeshComponents, true);
+
+			if (StaticMeshComponents.Num() > 0)
+			{
+				return true;
+			}
+		}
+	}
+	return false;
+}
+
+void UFractureActionTool::AddSingleRootNodeIfRequired(UGeometryCollection* GeometryCollectionObject)
+{
+	TSharedPtr<FGeometryCollection, ESPMode::ThreadSafe> GeometryCollectionPtr = GeometryCollectionObject->GetGeometryCollection();
+	if (FGeometryCollection* GeometryCollection = GeometryCollectionPtr.Get())
+	{
+		if (FGeometryCollectionClusteringUtility::ContainsMultipleRootBones(GeometryCollection))
+		{
+			FGeometryCollectionClusteringUtility::ClusterAllBonesUnderNewRoot(GeometryCollection);
+		}
+	}
+}
+
+void UFractureActionTool::AddAdditionalAttributesIfRequired(UGeometryCollection* GeometryCollectionObject)
+{
+	TSharedPtr<FGeometryCollection, ESPMode::ThreadSafe> GeometryCollectionPtr = GeometryCollectionObject->GetGeometryCollection();
+	if (FGeometryCollection* GeometryCollection = GeometryCollectionPtr.Get())
+	{
+		if (!GeometryCollection->HasAttribute("Level", FGeometryCollection::TransformGroup))
+		{
+			FGeometryCollectionClusteringUtility::UpdateHierarchyLevelOfChildren(GeometryCollection, -1);
+		}
+	}
+}
+
+void UFractureActionTool::GetSelectedGeometryCollectionComponents(TSet<UGeometryCollectionComponent*>& GeomCompSelection)
+{
+	USelection* SelectionSet = GEditor->GetSelectedActors();
+	TArray<AActor*> SelectedActors;
+	SelectedActors.Reserve(SelectionSet->Num());
+	SelectionSet->GetSelectedObjects(SelectedActors);
+
+	GeomCompSelection.Empty(SelectionSet->Num());
+
+	for (AActor* Actor : SelectedActors)
+	{
+		TInlineComponentArray<UGeometryCollectionComponent*> GeometryCollectionComponents;
+		Actor->GetComponents(GeometryCollectionComponents);
+		GeomCompSelection.Append(GeometryCollectionComponents);
+	}
+}
+
+
+void UFractureModalTool::Execute(TWeakPtr<FFractureEditorModeToolkit> InToolkit)
+{
+	if (!InToolkit.IsValid())
+	{
+		return;
+	}
+
+	TArray<FFractureToolContext> FractureContexts;
+	GetFractureContexts(FractureContexts);
+
+	TArray<UGeometryCollectionComponent*> NewComponents;
+	for (FFractureToolContext& FractureContext : FractureContexts)
+	{	
+		FractureContext.FracturedGeometryCollection->Modify();
+
+		ExecuteFracture(FractureContext);
+
+		{
+			TSharedPtr<FGeometryCollection, ESPMode::ThreadSafe> GeometryCollectionPtr = FractureContext.FracturedGeometryCollection->GetGeometryCollection();
+			FGeometryCollection* OutGeometryCollection = GeometryCollectionPtr.Get();
+			FGeometryCollectionClusteringUtility::UpdateHierarchyLevelOfChildren(OutGeometryCollection, -1);
+
+			// Update Nanite resource data to correctly reflect modified geometry collection data
+			{
+				FractureContext.FracturedGeometryCollection->ReleaseResources();
+
+				if (FractureContext.FracturedGeometryCollection->EnableNanite)
+				{
+					FractureContext.FracturedGeometryCollection->NaniteData = UGeometryCollection::CreateNaniteData(OutGeometryCollection);
+				}
+				else
+				{
+					FractureContext.FracturedGeometryCollection->NaniteData = MakeUnique<FGeometryCollectionNaniteData>();
+				}
+
+				FractureContext.FracturedGeometryCollection->InitResources();
+			}
+		}
+
+		UGeometryCollectionComponent* GeometryCollectionComponent = Cast<UGeometryCollectionComponent>(FractureContext.OriginalPrimitiveComponent);
+		NewComponents.AddUnique(GeometryCollectionComponent);
+		FGeometryCollectionEdit GCEdit = GeometryCollectionComponent->EditRestCollection();
+		UGeometryCollection* GCObject = GCEdit.GetRestCollection();
+		TSharedPtr<FGeometryCollection, ESPMode::ThreadSafe> GeometryCollectionPtr = GCObject->GetGeometryCollection();
+		FGeometryCollectionClusteringUtility::UpdateHierarchyLevelOfChildren(GeometryCollectionPtr.Get(), -1);
+
+		FScopedColorEdit EditBoneColor(GeometryCollectionComponent, true);
+		EditBoneColor.SelectBones(GeometryCollection::ESelectionMode::None);
+		InToolkit.Pin()->SetBoneSelection(GeometryCollectionComponent, EditBoneColor.GetSelectedBones(), true);
+
+		InToolkit.Pin()->UpdateExplodedVectors(GeometryCollectionComponent);
+
+		GeometryCollectionComponent->MarkRenderDynamicDataDirty();
+		GeometryCollectionComponent->MarkRenderStateDirty();
+		
+	}
+
+	InToolkit.Pin()->SetOutlinerComponents(NewComponents);
+}
+
+bool UFractureModalTool::CanExecute() const
+{
+	if (!IsGeometryCollectionSelected())
+	{
+		return false;
+	}
+
+	if (IsStaticMeshSelected())
+	{
+		return false;
+	}
+
+	return true;
+}
+
+void UFractureModalTool::PostEditChangeChainProperty(struct FPropertyChangedChainEvent& PropertyChangedEvent)
+{
+	FractureContextChanged();
+}
+
+void UFractureModalTool::GetFractureContexts(TArray<FFractureToolContext>& FractureContexts) const
+{
+	TSet<UGeometryCollectionComponent*> GeometryCollectionComponents;
+	GetSelectedGeometryCollectionComponents(GeometryCollectionComponents);
+
+	for (UGeometryCollectionComponent* GeometryCollectionComponent : GeometryCollectionComponents)
+	{
+		FGeometryCollectionEdit RestCollection = GeometryCollectionComponent->EditRestCollection(GeometryCollection::EEditUpdate::None);
+		UGeometryCollection* FracturedGeometryCollection = RestCollection.GetRestCollection();
+		if (FracturedGeometryCollection == nullptr)
+		{
+			continue;
+		}
+
+		TSharedPtr<FGeometryCollection, ESPMode::ThreadSafe> GeometryCollectionPtr = FracturedGeometryCollection->GetGeometryCollection();
+		FGeometryCollection* OutGeometryCollection = GeometryCollectionPtr.Get();
+
+		TArray<int32> FilteredBones = FilterBones(GeometryCollectionComponent->GetSelectedBones(), OutGeometryCollection);
 
+
+		const TManagedArray<FTransform>& Transform = OutGeometryCollection->GetAttribute<FTransform>("Transform", FGeometryCollection::TransformGroup);
+		const TManagedArray<int32>& TransformToGeometryIndex = OutGeometryCollection->GetAttribute<int32>("TransformToGeometryIndex", FGeometryCollection::TransformGroup);
+		const TManagedArray<FBox>& BoundingBoxes = OutGeometryCollection->GetAttribute<FBox>("BoundingBox", FGeometryCollection::GeometryGroup);
+
+		TArray<FTransform> Transforms;
+		GeometryCollectionAlgo::GlobalMatrices(Transform, OutGeometryCollection->Parent, Transforms);
+
+
+		TMap<int32, FBox> BoundsToBone;
+		for (int32 Idx = 0, ni = FracturedGeometryCollection->NumElements(FGeometryCollection::TransformGroup); Idx < ni; ++Idx)
+		{
+			if (TransformToGeometryIndex[Idx] > -1)
+			{
+				ensure(TransformToGeometryIndex[Idx] > -1);
+				BoundsToBone.Add(Idx, BoundingBoxes[TransformToGeometryIndex[Idx]].TransformBy(Transforms[Idx]));
+			}
+		}
+
+		GenerateFractureToolContext(
+			GeometryCollectionComponent->GetOwner(),
+			GeometryCollectionComponent,
+			FracturedGeometryCollection,
+			FilteredBones,
+			BoundsToBone,
+			TransformToGeometryIndex,
+			-1,
+			FractureContexts);
+
+	}
+}
+
+TArray<int32> UFractureModalTool::FilterBones(const TArray<int32>& SelectedBonesOriginal, const FGeometryCollection* const GeometryCollection) const
+{
+	// Keep only leaf nodes
+	const TManagedArray<TSet<int32>>& Children = GeometryCollection->GetAttribute<TSet<int32>>("Children", FGeometryCollection::TransformGroup);
+
+	TArray<int32> SelectedBones;
+	SelectedBones.Reserve(SelectedBonesOriginal.Num());
+	for (int32 BoneIndex : SelectedBonesOriginal)
+	{
+		if (Children[BoneIndex].Num() == 0)
+		{
+			SelectedBones.Add(BoneIndex);
+		}
+	}
+
+	return SelectedBones;
+}
+
+void UFractureModalTool::GenerateFractureToolContext(
+	AActor* InActor, 
+	UGeometryCollectionComponent* InComponent, 
+	UGeometryCollection* InGeometryCollection, 
+	const TArray<int32>& InSelectedBones, 
+	TMap<int32, FBox>& InBoundsToBone,
+	const TManagedArray<int32>& TransformToGeometryIndex,
+	int32 RandomSeed, 
+	TArray<FFractureToolContext>& OutFractureContexts
+) const
+{
+	OutFractureContexts.AddDefaulted();
+	FFractureToolContext& FractureContext = OutFractureContexts.Last();
+	FractureContext.RandomSeed = FMath::Rand();
+	if (RandomSeed > -1)
+	{
+		// make sure it's unique for each context if it's specified.
+		FractureContext.RandomSeed = RandomSeed + OutFractureContexts.Num();
+	}
+
+	FractureContext.OriginalActor = InActor;
+	FractureContext.Transform = InActor->GetActorTransform();
+	FractureContext.OriginalPrimitiveComponent = InComponent;
+	FractureContext.FracturedGeometryCollection = InGeometryCollection;
+	FractureContext.SelectedBones = InSelectedBones;
+
+	FractureContext.Bounds = FBox(ForceInit);
+	for (int32 BoneIndex : FractureContext.SelectedBones)
+	{
+		if (TransformToGeometryIndex[BoneIndex] > INDEX_NONE)
+		{
+			FractureContext.Bounds += InBoundsToBone[BoneIndex];
+		}
+	}
+}
diff --git a/Engine/Plugins/Experimental/ChaosEditor/Source/FractureEditor/Public/FractureTool.h b/Engine/Plugins/Experimental/ChaosEditor/Source/FractureEditor/Public/FractureTool.h
--- a/Engine/Plugins/Experimental/ChaosEditor/Source/FractureEditor/Public/FractureTool.h
+++ b/Engine/Plugins/Experimental/ChaosEditor/Source/FractureEditor/Public/FractureTool.h
@@ -8,126 +8,142 @@
 
 #include "Framework/Commands/UICommandInfo.h"
 #include "FractureEditorCommands.h"
+#include "FractureEditorModeToolkit.h"
 
 #include "FractureTool.generated.h"
 
 class UGeometryCollection;
-class UFractureTool;
+class UFractureModalTool;
 
-struct FFractureContext
+template <typename T>
+class TManagedArray;
+
+DECLARE_LOG_CATEGORY_EXTERN(LogFractureTool, Log, All);
+
+UCLASS(Abstract, config = EditorPerProjectUserSettings)
+class UFractureToolSettings : public UObject
+{
+	GENERATED_BODY()
+public:
+	UFractureToolSettings(const FObjectInitializer& ObjInit) : Super(ObjInit) {}
+
+	void PostEditChangeProperty(FPropertyChangedEvent& PropertyChangedEvent) override;
+	void PostEditChangeChainProperty(struct FPropertyChangedChainEvent& PropertyChangedEvent) override;
+
+	UPROPERTY()
+	UFractureModalTool* OwnerTool;
+};
+
+
+struct FFractureToolContext
 {
 	AActor* OriginalActor;
 	UPrimitiveComponent* OriginalPrimitiveComponent;
-
 	UGeometryCollection* FracturedGeometryCollection;
 	FString ParentName;
 	FTransform Transform;
 	FBox Bounds;
-// 	FVector InBoundsOffset;
-// 	int32 FracturedChunkIndex;
 	int32 RandomSeed;
 	TArray<int32> SelectedBones;
 };
 
-/** Settings specifically related to the one-time destructive fracturing of a mesh **/
-UCLASS(config = EditorPerProjectUserSettings)
-class UFractureCommonSettings: public UObject
+
+/** Tools derived from this class should require parameter inputs from the user, only the bone selection. */
+UCLASS(Abstract)
+class UFractureActionTool : public UObject
 {
-	GENERATED_BODY()
 public:
-	UFractureCommonSettings(const FObjectInitializer& ObjInit);
-
-#if WITH_EDITOR
-
-	void PostEditChangeProperty(FPropertyChangedEvent& PropertyChangedEvent) override;
-	void PostEditChangeChainProperty(struct FPropertyChangedChainEvent& PropertyChangedEvent) override;
+	GENERATED_BODY()
 
-#endif
-	/** Random number generator seed for repeatability */
-	UPROPERTY(EditAnywhere, Category = CommonFracture, meta = (DisplayName = "Random Seed", UIMin = "-1", UIMax = "1000", ClampMin = "-1"))
-	int32 RandomSeed;
+	UFractureActionTool(const FObjectInitializer& ObjInit) : Super(ObjInit) {}
 
-	/** Chance to shatter each mesh.  Useful when shattering multiple selected meshes.  */
-	UPROPERTY(EditAnywhere, Category = CommonFracture, meta = (DisplayName = "Chance To Fracture Per Mesh", UIMin = "0.0", UIMax = "1.0", ClampMin = "0.0", ClampMax = "1.0"))
-	float ChanceToFracture;
+	/** This is the Text that will appear on the tool button to execute the tool **/
+	virtual FText GetDisplayText() const { return FText(); }
+	virtual FText GetTooltipText() const { return FText(); }
 
-	/** Generate a fracture pattern across all selected meshes.  */
-	UPROPERTY(EditAnywhere, Category = CommonFracture, meta = (DisplayName = "Group Fracture"))
-	bool bGroupFracture;
+	virtual FSlateIcon GetToolIcon() const { return FSlateIcon(); }
 
-	/** Generate a fracture pattern across all selected meshes.  */
-	UPROPERTY(EditAnywhere, Category = CommonFracture, meta = (DisplayName = "Draw Sites"))
-	bool bDrawSites;
+	/** Executes the command.  Derived types need to be implemented in a thread safe way*/
+	virtual void Execute(TWeakPtr<FFractureEditorModeToolkit> InToolkit) {}
+	virtual bool CanExecute() const;
 
-	/** Generate a fracture pattern across all selected meshes.  */
-	UPROPERTY(EditAnywhere, Category = CommonFracture, meta = (DisplayName = "Draw Diagram"))
-	bool bDrawDiagram;
+	/** Gets the UI command info for this command */
+	const TSharedPtr<FUICommandInfo>& GetUICommandInfo() const;
 
-	/** Size of the noise displacement in centimeters */
-	UPROPERTY(EditAnywhere, Category = Noise, meta = (UIMin = "0.0"))
-	float Amplitude;
+	virtual void RegisterUICommand(FFractureEditorCommands* BindingContext) {}
 
-	/** Period of the Perlin noise.  Smaller values will create noise faces that are smoother */
-	UPROPERTY(EditAnywhere, Category = Noise)
-	float Frequency;
+	virtual void GetFractureToolContexts(TArray<FFractureToolContext>& FractureToolContexts) const {}
 
-	/** Number of fractal layers of Perlin noise to apply.  Smaller values (1 or 2) will create noise that looks like gentle rolling hills, while larger values (> 4) will tend to look more like craggy mountains */
-	UPROPERTY(EditAnywhere, Category = Noise, meta = (UIMin = "1"))
-	int32 OctaveNumber;
+	
+protected:
+	static bool IsStaticMeshSelected();
+	static bool IsGeometryCollectionSelected();
+	static void AddSingleRootNodeIfRequired(UGeometryCollection* GeometryCollectionObject);
+	static void AddAdditionalAttributesIfRequired(UGeometryCollection* GeometryCollectionObject);
+	static void GetSelectedGeometryCollectionComponents(TSet<UGeometryCollectionComponent*>& GeomCompSelection);
 
-	/** Spacing between vertices on cut surfaces, where noise is added.  Larger spacing between vertices will create more efficient meshes with fewer triangles, but less resolution to see the shape of the added noise  */
-	UPROPERTY(EditAnywhere, Category = Noise, meta = (UIMin = "1", ClampMin = "0.1"))
-	float SurfaceResolution;
+protected:
+	TSharedPtr<FUICommandInfo> UICommandInfo;
 
-	UPROPERTY()
-	UFractureTool *OwnerTool;
 };
 
-UCLASS(Abstract, Blueprintable, BlueprintType)
-class UFractureTool : public UObject
+
+/** Tools derived from this class provide parameter details and operate modally. */
+UCLASS(Abstract)
+class UFractureModalTool : public UFractureActionTool
 {
 public:
 	GENERATED_BODY()
 
-	UFractureTool(const FObjectInitializer& ObjInit) : Super(ObjInit) {}
+	UFractureModalTool(const FObjectInitializer& ObjInit) : Super(ObjInit) {}
 
-	/** This is the Text that will appear on the tool button to execute the fracture **/
-	virtual FText GetDisplayText() const  { return FText(); }
-	virtual FText GetTooltipText() const { return FText(); }
+	virtual TArray<UObject*> GetSettingsObjects() const { return TArray<UObject*>(); }
+	virtual void PostEditChangeProperty(FPropertyChangedEvent& PropertyChangedEvent) {}
 
 	/** This is the Text that will appear on the button to execute the fracture **/
-	virtual FText GetApplyText() const { return FText(NSLOCTEXT("Fracture", "ExecuteFracture", "Fracture")); }
+	virtual FText GetApplyText() const { return FText(); }
 
-	// TODO: What's the correct thing to return here?
-	virtual FSlateIcon GetToolIcon() const { return FSlateIcon(); }
+	/** Executes the command.  Derived types need to be implemented in a thread safe way*/
+	virtual void Execute(TWeakPtr<FFractureEditorModeToolkit> InToolkit) override;
+	virtual bool CanExecute() const override;
 
-	virtual TArray<UObject*> GetSettingsObjects() const { return TArray<UObject*>(); }
-
-	virtual void PostEditChangeProperty(FPropertyChangedEvent& PropertyChangedEvent) {}
+	virtual void ExecuteFracture(const FFractureToolContext& FractureContext) {}
 
 	/** Draw callback from edmode*/
 	virtual void Render(const FSceneView* View, FViewport* Viewport, FPrimitiveDrawInterface* PDI) {}
 
-	/** Executes the fracture command.  Derived types need to be implemented in a thread safe way*/
+	virtual void PostEditChangeChainProperty(struct FPropertyChangedChainEvent& PropertyChangedEvent) override;
 	virtual void FractureContextChanged() {}
-	virtual void ExecuteFracture(const FFractureContext& FractureContext) {}
-	virtual bool CanExecuteFracture() const { return true; }
 
+protected:
 
-	// Fracture Settings
-	UPROPERTY(EditAnywhere, Category = Slicing)
-	UFractureCommonSettings* CommonSettings;
+	virtual void GetFractureContexts(TArray<FFractureToolContext>& FractureContexts) const;
+	virtual TArray<int32> FilterBones(const TArray<int32>& SelectedBonesOriginal, const FGeometryCollection* const GeometryCollection) const;
 
+	void GenerateFractureToolContext(
+		AActor* InActor,
+		UGeometryCollectionComponent* InComponent,
+		UGeometryCollection* InGeometryCollection,
+		const TArray<int32>& InSelectedBones,
+		TMap<int32, FBox>& InBoundsToBone,
+		const TManagedArray<int32>& TransformToGeometryIndex,
+		int32 RandomSeed,
+		TArray<FFractureToolContext>& OutFractureContexts
+	) const;
 
-	/** Gets the UI command info for this command */
-	const TSharedPtr<FUICommandInfo>& GetUICommandInfo() const;
+};
 
-	virtual void RegisterUICommand(FFractureEditorCommands* BindingContext) {}
 
-	// virtual FUIAction MakeUIAction( class IMeshEditorModeUIContract& MeshEditorMode ) PURE_VIRTUAL(,return FUIAction(););
+/** Tools derived from this class provide parameter details, operate modally and use a viewport manipulator to set certain parameters. */
+UCLASS(Abstract)
+class UFractureInteractiveTool : public UFractureModalTool
+{
+public:
+	GENERATED_BODY()
 
-protected:
+	UFractureInteractiveTool(const FObjectInitializer& ObjInit) : Super(ObjInit) {}
 
-	TSharedPtr<FUICommandInfo> UICommandInfo;
+	// #todo (bmiller) implement interactive widgets
 
 };
+
