FractureEditor -- fixed a warning that LOCTEXT_NAMESPACE wasn't undef'd at the end of the file.

diff --git a/Engine/Plugins/Experimental/ChaosEditor/Source/FractureEditor/Private/FractureToolSelectors.cpp b/Engine/Plugins/Experimental/ChaosEditor/Source/FractureEditor/Private/FractureToolSelectors.cpp
--- a/Engine/Plugins/Experimental/ChaosEditor/Source/FractureEditor/Private/FractureToolSelectors.cpp
+++ b/Engine/Plugins/Experimental/ChaosEditor/Source/FractureEditor/Private/FractureToolSelectors.cpp
@@ -213,4 +213,6 @@ void UFractureToolSelectInvert::Execute(TWeakPtr<FFractureEditorModeToolkit> InT
 	{
 		SelectByMode(InToolkit.Pin().Get(), GeometryCollection::ESelectionMode::InverseGeometry);
 	}
-}
\ No newline at end of file
+}
+
+#undef LOCTEXT_NAMESPACE
\ No newline at end of file
