GeometryCollection -- implemented per bone initial dynamic state in Rest Collection.
FractureEditor -- SetInitialDynamicState tool

diff --git a/Engine/Plugins/Experimental/ChaosEditor/Source/FractureEditor/Private/FractureEditorCommands.h b/Engine/Plugins/Experimental/ChaosEditor/Source/FractureEditor/Private/FractureEditorCommands.h
--- a/Engine/Plugins/Experimental/ChaosEditor/Source/FractureEditor/Private/FractureEditorCommands.h
+++ b/Engine/Plugins/Experimental/ChaosEditor/Source/FractureEditor/Private/FractureEditorCommands.h
@@ -47,6 +47,9 @@ class FFractureEditorCommands : public TCommands<FFractureEditorCommands>
 		TSharedPtr< FUICommandInfo > Slice;
 		TSharedPtr< FUICommandInfo > Brick;
 		TSharedPtr< FUICommandInfo > Texture;
+
+		// Property Commands
+		TSharedPtr< FUICommandInfo > SetInitialDynamicState;
 		
 };
 
diff --git a/Engine/Plugins/Experimental/ChaosEditor/Source/FractureEditor/Private/FractureEditorModeToolkit.cpp b/Engine/Plugins/Experimental/ChaosEditor/Source/FractureEditor/Private/FractureEditorModeToolkit.cpp
--- a/Engine/Plugins/Experimental/ChaosEditor/Source/FractureEditor/Private/FractureEditorModeToolkit.cpp
+++ b/Engine/Plugins/Experimental/ChaosEditor/Source/FractureEditor/Private/FractureEditorModeToolkit.cpp
@@ -150,6 +150,7 @@ FFractureViewSettingsCustomization::CustomizeDetails(IDetailLayoutBuilder& Detai
 	];
 };
 
+
 FFractureEditorModeToolkit::FFractureEditorModeToolkit()
 	: ActiveTool(nullptr)
 {
@@ -323,11 +324,11 @@ void FFractureEditorModeToolkit::OnObjectPostEditChange( UObject* Object, FPrope
 		else if ( PropertyChangedEvent.Property->GetFName() == GET_MEMBER_NAME_CHECKED(UFractureSettings, FractureLevel))
 		{
 			OnLevelViewValueChanged();
-		} 
+		}
 	}
 }
 
-const TArray<FName> FFractureEditorModeToolkit::PaletteNames = { FName(TEXT("Generate")), FName(TEXT("Select")), FName(TEXT("Fracture")), FName(TEXT("Cluster")) };
+const TArray<FName> FFractureEditorModeToolkit::PaletteNames = { FName(TEXT("Generate")), FName(TEXT("Select")), FName(TEXT("Fracture")), FName(TEXT("Cluster")), FName(TEXT("Properties")) };
 
 FText FFractureEditorModeToolkit::GetToolPaletteDisplayName(FName Palette) const
 { 
@@ -383,6 +384,10 @@ void FFractureEditorModeToolkit::BuildToolPalette(FName PaletteIndex, class FToo
 		ToolbarBuilder.AddToolBarButton(Commands.Uncluster);
 		ToolbarBuilder.AddToolBarButton(Commands.MoveUp);
 	}
+	else if (PaletteIndex == TEXT("Properties"))
+	{
+		ToolbarBuilder.AddToolBarButton(Commands.SetInitialDynamicState);
+	}
 }
 
 void FFractureEditorModeToolkit::BindCommands()
@@ -794,6 +799,7 @@ void FFractureEditorModeToolkit::SetOutlinerComponents(const TArray<UGeometryCol
 			FGeometryCollectionClusteringUtility::UpdateHierarchyLevelOfChildren(GeometryCollectionPtr.Get(), -1);
 			UpdateExplodedVectors(Component);
 		}
+
 	}
 
 	if (OutlinerView)
@@ -810,6 +816,7 @@ void FFractureEditorModeToolkit::SetOutlinerComponents(const TArray<UGeometryCol
 void FFractureEditorModeToolkit::SetBoneSelection(UGeometryCollectionComponent* InRootComponent, const TArray<int32>& InSelectedBones, bool bClearCurrentSelection)
 {
 	OutlinerView->SetBoneSelection(InRootComponent, InSelectedBones, bClearCurrentSelection);
+	
 	if (ActiveTool != nullptr)
 	{
 		ActiveTool->FractureContextChanged();
@@ -916,6 +923,7 @@ bool FFractureEditorModeToolkit::IsSelectedActorsInEditorWorld()
 	return true;
 }
 
+
 bool GetValidGeoCenter(const TManagedArray<int32>& TransformToGeometryIndex, const TArray<FTransform>& Transforms, const TManagedArray<TSet<int32>>& Children, const TManagedArray<FBox>& BoundingBox, int32 TransformIndex, FVector& OutGeoCenter )
 {
 	if (TransformToGeometryIndex[TransformIndex] > -1)
@@ -1065,6 +1073,7 @@ void FFractureEditorModeToolkit::OnOutlinerBoneSelectionChanged(UGeometryCollect
 	if(SelectedBones.Num())
 	{
 		FFractureSelectionTools::ToggleSelectedBones(RootComponent, SelectedBones, true);
+		OutlinerView->SetBoneSelection(RootComponent, SelectedBones, true);
 	}
 	else
 	{
diff --git a/Engine/Plugins/Experimental/ChaosEditor/Source/FractureEditor/Private/FractureToolProperties.cpp b/Engine/Plugins/Experimental/ChaosEditor/Source/FractureEditor/Private/FractureToolProperties.cpp
new file mode 100644
--- /dev/null
+++ b/Engine/Plugins/Experimental/ChaosEditor/Source/FractureEditor/Private/FractureToolProperties.cpp
@@ -0,0 +1,82 @@
+// Copyright Epic Games, Inc. All Rights Reserved.
+
+#include "FractureToolProperties.h"
+
+#include "GeometryCollection/GeometryCollectionComponent.h"
+
+
+#define LOCTEXT_NAMESPACE "FractureProperties"
+
+
+UFractureToolSetInitialDynamicState::UFractureToolSetInitialDynamicState(const FObjectInitializer& ObjInit)
+	: Super(ObjInit)
+{
+	StateSettings = NewObject<UFractureInitialDynamicStateSettings>(GetTransientPackage(), UFractureInitialDynamicStateSettings::StaticClass());
+	StateSettings->OwnerTool = this;
+}
+
+FText UFractureToolSetInitialDynamicState::GetDisplayText() const
+{
+	return FText(NSLOCTEXT("Fracture", "FractureToolSetInitialDynamicState", "State"));
+}
+
+FText UFractureToolSetInitialDynamicState::GetTooltipText() const
+{
+	return FText(NSLOCTEXT("Fracture", "FractureToolSetInitialDynamicStateToolTip", \
+		"Override initial dynamic state for selected bones. If the component's Object Type is set to Dynamic, the solver will use this override state instead. Setting a bone to Kinematic will have the effect of anchoring the bone to world space, for instance."));
+}
+
+FText UFractureToolSetInitialDynamicState::GetApplyText() const
+{
+	return FText(NSLOCTEXT("Fracture", "ExecuteSetInitialDynamicState", "Set State"));
+}
+
+FSlateIcon UFractureToolSetInitialDynamicState::GetToolIcon() const
+{
+	return FSlateIcon("FractureEditorStyle", "FractureEditor.SetInitialDynamicState");
+}
+
+TArray<UObject*> UFractureToolSetInitialDynamicState::GetSettingsObjects() const
+{
+	TArray<UObject*> Settings;
+	Settings.Add(StateSettings);
+	return Settings;
+}
+
+void UFractureToolSetInitialDynamicState::RegisterUICommand(FFractureEditorCommands* BindingContext)
+{
+	UI_COMMAND_EXT(BindingContext, UICommandInfo, "SetInitialDynamicState", "State", \
+		"Override initial dynamic state for selected bones. If the component's Object Type is set to Dynamic, the solver will use this override state instead. Setting a bone to Kinematic will have the effect of anchoring the bone to world space, for instance.", \
+		EUserInterfaceActionType::ToggleButton, FInputChord());
+	BindingContext->SetInitialDynamicState = UICommandInfo;
+}
+
+void UFractureToolSetInitialDynamicState::Execute(TWeakPtr<FFractureEditorModeToolkit> InToolkit)
+{
+	if (InToolkit.IsValid())
+	{
+		TSet<UGeometryCollectionComponent*> GeomCompSelection;
+		GetSelectedGeometryCollectionComponents(GeomCompSelection);
+		for (UGeometryCollectionComponent* GeometryCollectionComponent : GeomCompSelection)
+		{
+			FGeometryCollectionEdit GCEdit = GeometryCollectionComponent->EditRestCollection();
+			if (UGeometryCollection* GCObject = GCEdit.GetRestCollection())
+			{
+				TSharedPtr<FGeometryCollection, ESPMode::ThreadSafe> GeometryCollectionPtr = GCObject->GetGeometryCollection();
+				if (FGeometryCollection* GeometryCollection = GeometryCollectionPtr.Get())
+				{
+					TManagedArray<int32>& InitialDynamicState = GeometryCollection->GetAttribute<int32>("InitialDynamicState", FGeometryCollection::TransformGroup);
+
+					TArray<int32> SelectedBones = GeometryCollectionComponent->GetSelectedBones();
+					for (int32 Index : SelectedBones)
+					{
+						InitialDynamicState[Index] = static_cast<int32>(StateSettings->InitialDynamicState);
+					}
+				}
+			}
+		}
+	}
+}
+			
+
+#undef LOCTEXT_NAMESPACE
diff --git a/Engine/Plugins/Experimental/ChaosEditor/Source/FractureEditor/Private/FractureToolProperties.h b/Engine/Plugins/Experimental/ChaosEditor/Source/FractureEditor/Private/FractureToolProperties.h
new file mode 100644
--- /dev/null
+++ b/Engine/Plugins/Experimental/ChaosEditor/Source/FractureEditor/Private/FractureToolProperties.h
@@ -0,0 +1,62 @@
+// Copyright Epic Games, Inc. All Rights Reserved.
+#pragma once
+
+#include "FractureTool.h"
+
+#include "FractureToolProperties.generated.h"
+
+
+UENUM(BlueprintType)
+enum class EDynamicStateOverrideEnum : uint8
+{
+	NoOverride = 0										  UMETA(DisplayName = "No Override"),
+	Sleeping = 1 /*Chaos::EObjectStateType::Sleeping*/    UMETA(DisplayName = "Sleeping"),
+	Kinematic = 2 /*Chaos::EObjectStateType::Kinematic*/  UMETA(DisplayName = "Kinematic"),
+	Static = 3    /*Chaos::EObjectStateType::Static*/     UMETA(DisplayName = "Static")
+};
+
+/** Settings specifically related to the one-time destructive fracturing of a mesh **/
+UCLASS(config = EditorPerProjectUserSettings)
+class UFractureInitialDynamicStateSettings : public UFractureToolSettings
+{
+public:
+
+	GENERATED_BODY()
+
+	UFractureInitialDynamicStateSettings(const FObjectInitializer& ObjInit)
+		: Super(ObjInit)
+		, InitialDynamicState(EDynamicStateOverrideEnum::Kinematic)
+	{}
+
+	/** Random number generator seed for repeatability */
+	UPROPERTY(EditAnywhere, Category = SetInitialDynamicState, meta = (DisplayName = "Initial Dynamic State"))
+	EDynamicStateOverrideEnum InitialDynamicState;
+
+};
+
+
+UCLASS(DisplayName = "State", Category = "FractureTools")
+class UFractureToolSetInitialDynamicState : public UFractureModalTool
+{
+public:
+	GENERATED_BODY()
+
+	UFractureToolSetInitialDynamicState(const FObjectInitializer& ObjInit);
+
+	using FGeometryCollectionPtr = TSharedPtr<FGeometryCollection, ESPMode::ThreadSafe>;
+
+	// UFractureTool Interface
+	virtual FText GetDisplayText() const override;
+	virtual FText GetTooltipText() const override;
+	virtual FText GetApplyText() const override;
+	virtual FSlateIcon GetToolIcon() const override;
+	virtual TArray<UObject*> GetSettingsObjects() const override;
+
+	virtual void RegisterUICommand(FFractureEditorCommands* BindingContext) override;
+
+	virtual void Execute(TWeakPtr<FFractureEditorModeToolkit> InToolkit) override;
+
+	UPROPERTY(EditAnywhere, Category = InitialDynamicState)
+	UFractureInitialDynamicStateSettings* StateSettings;
+
+};
diff --git a/Engine/Source/Runtime/Experimental/Chaos/Private/GeometryCollection/GeometryCollection.cpp b/Engine/Source/Runtime/Experimental/Chaos/Private/GeometryCollection/GeometryCollection.cpp
--- a/Engine/Source/Runtime/Experimental/Chaos/Private/GeometryCollection/GeometryCollection.cpp
+++ b/Engine/Source/Runtime/Experimental/Chaos/Private/GeometryCollection/GeometryCollection.cpp
@@ -47,6 +47,7 @@ void FGeometryCollection::Construct()
 	AddExternalAttribute<int32>("TransformToGeometryIndex", FTransformCollection::TransformGroup, TransformToGeometryIndex);
 	AddExternalAttribute<int32>("SimulationType", FTransformCollection::TransformGroup, SimulationType);
 	AddExternalAttribute<int32>("StatusFlags", FTransformCollection::TransformGroup, StatusFlags);
+	AddExternalAttribute<int32>("InitialDynamicState", FTransformCollection::TransformGroup, InitialDynamicState);
 
 	// Vertices Group
 	AddExternalAttribute<FVector>("Vertex", FGeometryCollection::VerticesGroup, Vertex);
@@ -88,6 +89,7 @@ void FGeometryCollection::SetDefaults(FName Group, uint32 StartSize, uint32 NumE
 			Parent[Idx] = FGeometryCollection::Invalid;
 			SimulationType[Idx] = FGeometryCollection::ESimulationTypes::FST_None;
 			StatusFlags[Idx] = 0;
+			InitialDynamicState[Idx] = static_cast<int32>(Chaos::EObjectStateType::Uninitialized);
 		}
 	}
 }
@@ -136,12 +138,14 @@ int32 FGeometryCollection::AppendGeometry(const FGeometryCollection & Element, i
 
 	const TManagedArray<int32>& ElementSimulationType = Element.SimulationType;
 	const TManagedArray<int32>& ElementStatusFlags = Element.StatusFlags;
+	const TManagedArray<int32>& ElementInitialDynamicState = Element.InitialDynamicState;
 
 	// --- TRANSFORM ---
 	for (int TransformIdx = 0; TransformIdx < NumNewTransforms; TransformIdx++)
 	{
 		SimulationType[TransformIdx + StartTransformIndex] = ElementSimulationType[TransformIdx];
 		StatusFlags[TransformIdx + StartTransformIndex] = ElementStatusFlags[TransformIdx];
+		InitialDynamicState[TransformIdx + StartTransformIndex] = ElementInitialDynamicState[TransformIdx];
 	}
 
 	// --- VERTICES GROUP ---
@@ -1365,6 +1369,7 @@ FGeometryCollection* FGeometryCollection::NewGeometryCollection(const TArray<flo
 	TManagedArray<TSet<int32>>& Children = RestCollection->Children;
 	TManagedArray<int32>& SimulationType = RestCollection->SimulationType;
 	TManagedArray<int32>& StatusFlags = RestCollection->StatusFlags;
+	TManagedArray<int32>& InitialDynamicState = RestCollection->InitialDynamicState;
 
 	// set the vertex information
 	for (int32 Idx = 0; Idx < NumNewVertices; ++Idx)
diff --git a/Engine/Source/Runtime/Experimental/Chaos/Public/GeometryCollection/GeometryCollection.h b/Engine/Source/Runtime/Experimental/Chaos/Public/GeometryCollection/GeometryCollection.h
--- a/Engine/Source/Runtime/Experimental/Chaos/Public/GeometryCollection/GeometryCollection.h
+++ b/Engine/Source/Runtime/Experimental/Chaos/Public/GeometryCollection/GeometryCollection.h
@@ -231,6 +231,7 @@ public:
 	TManagedArray<int32>		TransformToGeometryIndex;
 	TManagedArray<int32>        SimulationType;
 	TManagedArray<int32>        StatusFlags;
+	TManagedArray<int32>		InitialDynamicState;
 
 	// Vertices Group
 	TManagedArray<FVector>		Vertex;
diff --git a/Engine/Source/Runtime/Experimental/GeometryCollectionEngine/Private/GeometryCollection/GeometryCollectionComponent.cpp b/Engine/Source/Runtime/Experimental/GeometryCollectionEngine/Private/GeometryCollection/GeometryCollectionComponent.cpp
--- a/Engine/Source/Runtime/Experimental/GeometryCollectionEngine/Private/GeometryCollection/GeometryCollectionComponent.cpp
+++ b/Engine/Source/Runtime/Experimental/GeometryCollectionEngine/Private/GeometryCollection/GeometryCollectionComponent.cpp
@@ -1950,18 +1950,28 @@ void UGeometryCollectionComponent::OnCreatePhysicsState()
 #endif
 			};
 
-			// @todo(temporary) : This is Temporary code for the collection to match the ObjectType
-			//                    attribute on initialization. Once proper per-object manipulation is 
-			//                    in place this code will need to be removed.
-			//
+			// If the Component is set to Dynamic, we look to the RestCollection for initial dynamic state override per transform.
 			TManagedArray<int32> & DynamicState = DynamicCollection->DynamicState;
+
 			if (ObjectType != EObjectStateTypeEnum::Chaos_Object_UserDefined)
 			{
-				for (int i = 0; i < DynamicState.Num(); i++)
+				if (RestCollection && (ObjectType == EObjectStateTypeEnum::Chaos_Object_Dynamic))
+				{
+					TManagedArray<int32>& InitialDynamicState = RestCollection->GetGeometryCollection()->InitialDynamicState;
+					for (int i = 0; i < DynamicState.Num(); i++)
+					{
+						DynamicState[i] = (InitialDynamicState[i] == static_cast<int32>(Chaos::EObjectStateType::Uninitialized)) ? static_cast<int32>(ObjectType) : InitialDynamicState[i];
+					}
+				}
+				else
 				{
-					DynamicState[i] = (int32)ObjectType;
+					for (int i = 0; i < DynamicState.Num(); i++)
+					{
+						DynamicState[i] = static_cast<int32>(ObjectType);
+					}
 				}
 			}
+
 			TManagedArray<bool> & Active = DynamicCollection->Active;
 			{
 				for (int i = 0; i < Active.Num(); i++)
@@ -1976,8 +1986,7 @@ void UGeometryCollectionComponent::OnCreatePhysicsState()
 					CollisionGroupArray[i] = CollisionGroup;
 				}
 			}
-			// end temporary 
-
+	
 			// Set up initial filter data for our particles
 			// #BGTODO We need a dummy body setup for now to allow the body instance to generate filter information. Change body instance to operate independently.
 			DummyBodySetup = NewObject<UBodySetup>(this, UBodySetup::StaticClass());
