Fracture Editor Mode
Refactored FractureToolContext to handle all the selection and filtering for all tools in the mode.
Implemented histogram view to suppliment Outliner view.
Fixed some problematic Clustering logic.
All tools except AutoFracture are now using Context filtered selection.
Flatten and ClusterMagnet can be run on a single branch.

diff --git a/Engine/Plugins/Experimental/ChaosEditor/Source/FractureEditor/Private/FractureEditorModeToolkit.cpp b/Engine/Plugins/Experimental/ChaosEditor/Source/FractureEditor/Private/FractureEditorModeToolkit.cpp
--- a/Engine/Plugins/Experimental/ChaosEditor/Source/FractureEditor/Private/FractureEditorModeToolkit.cpp
+++ b/Engine/Plugins/Experimental/ChaosEditor/Source/FractureEditor/Private/FractureEditorModeToolkit.cpp
@@ -45,6 +45,7 @@
 #include "PlanarCut.h"
 #include "FractureToolAutoCluster.h" 
 #include "SGeometryCollectionOutliner.h"
+#include "SGeometryCollectionHistogram.h"
 #include "FractureSelectionTools.h"
 #include "GeometryCollection/GeometryCollectionAlgo.h"
 #include "GeometryCollection/GeometryCollectionClusteringUtility.h"
@@ -52,6 +53,10 @@
 #include "LevelEditorViewport.h"
 #include "Widgets/Layout/SUniformGridPanel.h"
 
+#include "Chaos/TriangleMesh.h"
+#include "PhysicsProxy/GeometryCollectionPhysicsProxy.h"
+#include "Chaos/MassProperties.h"
+
 #include "DetailWidgetRow.h"
 #include "DetailCategoryBuilder.h"
 #include "DetailLayoutBuilder.h"
@@ -150,6 +155,39 @@ FFractureViewSettingsCustomization::CustomizeDetails(IDetailLayoutBuilder& Detai
 	];
 };
 
+FHistogramSettingsCustomization::FHistogramSettingsCustomization(FFractureEditorModeToolkit* InToolkit)
+	: Toolkit(InToolkit)
+{
+
+}
+
+TSharedRef<IDetailCustomization> FHistogramSettingsCustomization::MakeInstance(FFractureEditorModeToolkit* InToolkit)
+{
+	return MakeShareable(new FHistogramSettingsCustomization(InToolkit));
+}
+
+void
+FHistogramSettingsCustomization::CustomizeDetails(IDetailLayoutBuilder& DetailBuilder)
+{
+};
+
+FOutlinerSettingsCustomization::FOutlinerSettingsCustomization(FFractureEditorModeToolkit* InToolkit)
+	: Toolkit(InToolkit)
+{
+
+}
+
+TSharedRef<IDetailCustomization> FOutlinerSettingsCustomization::MakeInstance(FFractureEditorModeToolkit* InToolkit)
+{
+	return MakeShareable(new FOutlinerSettingsCustomization(InToolkit));
+}
+
+void
+FOutlinerSettingsCustomization::CustomizeDetails(IDetailLayoutBuilder& DetailBuilder)
+{
+};
+
+
 
 FFractureEditorModeToolkit::FFractureEditorModeToolkit()
 	: ActiveTool(nullptr)
@@ -190,9 +228,52 @@ void FFractureEditorModeToolkit::Init(const TSharedPtr<IToolkitHost>& InitToolki
 	Settings.Add(GetMutableDefault<UFractureSettings>());
 	DetailsView->SetObjects(Settings);
 
+	HistogramDetailsView = EditModule.CreateDetailView(DetailsViewArgs);
+	EditModule.RegisterCustomClassLayout("HistogramSettings", FOnGetDetailCustomizationInstance::CreateStatic(&FHistogramSettingsCustomization::MakeInstance, this));
+	HistogramDetailsView->SetObject(GetMutableDefault<UHistogramSettings>());
+
+	OutlinerDetailsView = EditModule.CreateDetailView(DetailsViewArgs);
+	EditModule.RegisterCustomClassLayout("OutlinerSettings", FOnGetDetailCustomizationInstance::CreateStatic(&FOutlinerSettingsCustomization::MakeInstance, this));
+	OutlinerDetailsView->SetObject(GetMutableDefault<UOutlinerSettings>());
+
 	float Padding = 4.0f;
 	FMargin MorePadding = FMargin(10.0f, 2.0f);
 
+	TSharedRef<SExpandableArea> HistogramExpander = SNew(SExpandableArea)
+	.AreaTitle(FText(LOCTEXT("Histogram", "Histogram")))
+	.HeaderPadding(FMargin(2.0, 2.0))
+	.Padding(MorePadding)
+	.BorderBackgroundColor(FLinearColor(.6, .6, .6, 1.0f))
+	.BodyBorderBackgroundColor(FLinearColor(1.0, 0.0, 0.0))
+	.AreaTitleFont(FEditorStyle::Get().GetFontStyle("HistogramDetailsView.CategoryFontStyle"))
+	.InitiallyCollapsed(true)
+	.BodyContent()
+	[
+		SNew(SVerticalBox)
+		+ SVerticalBox::Slot()
+		.AutoHeight()
+		[
+			HistogramDetailsView.ToSharedRef()
+		]
+		+ SVerticalBox::Slot()
+		.AutoHeight()
+		[
+			SNew(SButton)
+			//.ButtonStyle(FAppStyle::Get(), "PrimaryButton")
+			.TextStyle(FAppStyle::Get(), "ButtonText")
+			.HAlign(HAlign_Center)
+			.ContentPadding(FMargin(10.f, Padding))
+			.OnClicked(this, &FFractureEditorModeToolkit::ResetHistogramSelection)
+			.IsEnabled(this, &FFractureEditorModeToolkit::CanResetFilter)
+			.Text(LOCTEXT("ResetFilterButton", "Reset Filter"))
+		]
+		+ SVerticalBox::Slot()
+		[
+			SAssignNew(HistogramView, SGeometryCollectionHistogram)
+			.OnBoneSelectionChanged(this, &FFractureEditorModeToolkit::OnHistogramBoneSelectionChanged)
+		]
+	];
+
 	TSharedRef<SExpandableArea> OutlinerExpander = SNew(SExpandableArea)
 	.AreaTitle(FText(LOCTEXT("Outliner", "Outliner")))
 	.HeaderPadding(FMargin(2.0, 2.0))
@@ -203,8 +284,36 @@ void FFractureEditorModeToolkit::Init(const TSharedPtr<IToolkitHost>& InitToolki
 	.AreaTitleFont(FEditorStyle::Get().GetFontStyle("DetailsView.CategoryFontStyle"))
 	.BodyContent()
 	[
-		SAssignNew(OutlinerView, SGeometryCollectionOutliner )
-		.OnBoneSelectionChanged(this, &FFractureEditorModeToolkit::OnOutlinerBoneSelectionChanged)
+		SNew(SVerticalBox)
+		+SVerticalBox::Slot()
+		[
+			SNew(SSplitter)
+			.Orientation(Orient_Vertical)
+			+ SSplitter::Slot()
+			.SizeRule(TAttribute<SSplitter::ESizeRule>::Create([this, HistogramExpander]() {
+				return HistogramExpander->IsExpanded() ? SSplitter::ESizeRule::FractionOfParent : SSplitter::ESizeRule::SizeToContent;
+				}))
+			.Value(1.f)
+			[
+				HistogramExpander
+			]
+			+ SSplitter::Slot()
+			.SizeRule(SSplitter::ESizeRule::FractionOfParent)
+			.Value(1.f)
+			[
+				SNew(SVerticalBox)
+				+ SVerticalBox::Slot()
+				.AutoHeight()
+				[
+					OutlinerDetailsView.ToSharedRef()
+				]
+				+ SVerticalBox::Slot()
+				[
+					SAssignNew(OutlinerView, SGeometryCollectionOutliner)
+					.OnBoneSelectionChanged(this, &FFractureEditorModeToolkit::OnOutlinerBoneSelectionChanged)
+				]
+			]
+		]	
 	];
 
 	TSharedRef<SExpandableArea> StatisticsExpander = SNew(SExpandableArea)
@@ -325,6 +434,22 @@ void FFractureEditorModeToolkit::OnObjectPostEditChange( UObject* Object, FPrope
 		{
 			OnLevelViewValueChanged();
 		}
+		else if (PropertyChangedEvent.Property->GetFName() == GET_MEMBER_NAME_CHECKED(UHistogramSettings, bSorted))
+		{
+			UHistogramSettings* HistogramSettings = GetMutableDefault<UHistogramSettings>();
+			HistogramView->RefreshView(HistogramSettings->bSorted);
+		}
+		else if (PropertyChangedEvent.Property->GetFName() == GET_MEMBER_NAME_CHECKED(UHistogramSettings, InspectedAttribute))
+		{
+			UHistogramSettings* HistogramSettings = GetMutableDefault<UHistogramSettings>();
+			HistogramView->InspectAttribute(HistogramSettings->InspectedAttribute);
+		}
+		else if (PropertyChangedEvent.Property->GetFName() == GET_MEMBER_NAME_CHECKED(UOutlinerSettings, ItemText))
+		{
+			UOutlinerSettings* OutlinerSettings = GetMutableDefault<UOutlinerSettings>();
+			OutlinerView->RegenerateItems();
+		}
+
 	}
 }
 
@@ -800,6 +925,7 @@ void FFractureEditorModeToolkit::SetOutlinerComponents(const TArray<UGeometryCol
 			UpdateExplodedVectors(Component);
 		}
 
+		UpdateGeometryComponentAttributes(Component);
 	}
 
 	if (OutlinerView)
@@ -807,6 +933,11 @@ void FFractureEditorModeToolkit::SetOutlinerComponents(const TArray<UGeometryCol
 		OutlinerView->SetComponents(InNewComponents);
 	}
 
+	if (HistogramView)
+	{
+		HistogramView->SetComponents(InNewComponents);
+	}
+
 	if (ActiveTool != nullptr)
 	{
 		ActiveTool->FractureContextChanged();
@@ -858,6 +989,17 @@ bool FFractureEditorModeToolkit::CanExecuteModal() const
 	return false;
 }
 
+FReply FFractureEditorModeToolkit::ResetHistogramSelection()
+{
+	HistogramView->ClearSelection();
+	return FReply::Handled();
+}
+
+bool FFractureEditorModeToolkit::CanResetFilter() const 
+{
+	return HistogramView->IsSelected();
+}
+
 void FFractureEditorModeToolkit::GetSelectedGeometryCollectionComponents(TSet<UGeometryCollectionComponent*>& GeomCompSelection)
 {
 	USelection* SelectionSet = GEditor->GetSelectedActors();
@@ -924,6 +1066,94 @@ bool FFractureEditorModeToolkit::IsSelectedActorsInEditorWorld()
 }
 
 
+void FFractureEditorModeToolkit::UpdateGeometryComponentAttributes(UGeometryCollectionComponent* Component)
+{
+	if (Component)
+	{
+		const UGeometryCollection* RestCollection = Component->GetRestCollection();
+		if (RestCollection)
+		{
+			FGeometryCollectionPtr GeometryCollection = RestCollection->GetGeometryCollection();
+
+			if (!GeometryCollection->HasAttribute("Volume", FTransformCollection::TransformGroup))
+			{
+				GeometryCollection->AddAttribute<Chaos::FReal>("Volume", FTransformCollection::TransformGroup);
+				UE_LOG(LogFractureTool, Warning, TEXT("Added Volume attribute to GeometryCollection."));
+			}
+
+			TArray<FTransform> Transform;
+			GeometryCollectionAlgo::GlobalMatrices(GeometryCollection->Transform, GeometryCollection->Parent, Transform);
+
+			const TManagedArray<FVector>& Vertex = GeometryCollection->Vertex;
+			const TManagedArray<int32>& BoneMap = GeometryCollection->BoneMap;
+
+			Chaos::TParticles<float, 3> MassSpaceParticles;
+			MassSpaceParticles.AddParticles(Vertex.Num());
+			for (int32 Idx = 0; Idx < Vertex.Num(); ++Idx)
+			{
+				MassSpaceParticles.X(Idx) = Transform[BoneMap[Idx]].TransformPosition(Vertex[Idx]);
+			}
+
+			// recalculate mass recursively from root nodes
+			const TManagedArray<int32>& Parent = GeometryCollection->Parent;
+			for (int32 Idx = 0; Idx < Parent.Num(); ++Idx)
+			{
+				if (Parent[Idx] == INDEX_NONE)
+				{
+					UpdateVolumes(GeometryCollection, MassSpaceParticles, Idx);
+				}
+			}
+		}
+	}
+	
+}
+
+void FFractureEditorModeToolkit::UpdateVolumes(FGeometryCollectionPtr GeometryCollection, const Chaos::TParticles<float, 3>& MassSpaceParticles, int32 TransformIndex)
+{
+	const TManagedArray<TSet<int32>>& Children = GeometryCollection->Children;
+	const TManagedArray<bool>& Visible = GeometryCollection->Visible;
+	const TManagedArray<int32>& FaceCount = GeometryCollection->FaceCount;
+	const TManagedArray<int32>& FaceStart = GeometryCollection->FaceStart;
+	const TManagedArray<int32>& TransformToGeometryIndex = GeometryCollection->TransformToGeometryIndex;
+	const TManagedArray<FIntVector>& Indices = GeometryCollection->Indices;
+	TManagedArray<float>& Volumes = GeometryCollection->GetAttribute<float>("Volume", FTransformCollection::TransformGroup);
+
+	if (Children[TransformIndex].Num() == 0) // leaf node
+	{
+		int32 GeometryIndex = TransformToGeometryIndex[TransformIndex];
+
+		if (ensureMsgf(GeometryIndex > INDEX_NONE, TEXT("Leaf node %d expected to map to geometry but did not."), TransformIndex))
+		{
+			TUniquePtr<Chaos::TTriangleMesh<float>> TriMesh(
+				CreateTriangleMesh(
+					FaceStart[GeometryIndex],
+					FaceCount[GeometryIndex],
+					Visible,
+					Indices,
+					true));
+
+			float Volume = 0.0;
+			Chaos::TVector<Chaos::FReal, 3> CenterOfMass;
+			Chaos::CalculateVolumeAndCenterOfMass(MassSpaceParticles, TriMesh->GetElements(), Volume, CenterOfMass);
+
+			// Since we're only interested in relative mass, we assume density = 1.0
+			Volumes[TransformIndex] = Volume;
+		}	
+	}
+	else
+	{
+		// Recurse to children and sum the volumes for this node
+		float LocalVolume = 0.0;
+		for (int32 ChildIndex : Children[TransformIndex])
+		{
+			UpdateVolumes(GeometryCollection, MassSpaceParticles, ChildIndex);
+			LocalVolume += Volumes[ChildIndex];
+		}
+
+		Volumes[TransformIndex] = LocalVolume;
+	}
+}
+
 bool GetValidGeoCenter(const TManagedArray<int32>& TransformToGeometryIndex, const TArray<FTransform>& Transforms, const TManagedArray<TSet<int32>>& Children, const TManagedArray<FBox>& BoundingBox, int32 TransformIndex, FVector& OutGeoCenter )
 {
 	if (TransformToGeometryIndex[TransformIndex] > -1)
@@ -1066,12 +1296,18 @@ void FFractureEditorModeToolkit::UpdateExplodedVectors(UGeometryCollectionCompon
 	}
 }
 
+void FFractureEditorModeToolkit::RegenerateHistogram()
+{
+	HistogramView->RegenerateNodes();
+}
+
 void FFractureEditorModeToolkit::OnOutlinerBoneSelectionChanged(UGeometryCollectionComponent* RootComponent, TArray<int32>& SelectedBones)
 {
 	FScopedTransaction Transaction(FractureTransactionContexts::SelectBoneContext, LOCTEXT("SelectGeometryCollectionBoneTransaction", "Select Bone"), RootComponent);
 
 	if(SelectedBones.Num())
 	{
+			
 		FFractureSelectionTools::ToggleSelectedBones(RootComponent, SelectedBones, true);
 		OutlinerView->SetBoneSelection(RootComponent, SelectedBones, true);
 	}
@@ -1089,6 +1325,12 @@ void FFractureEditorModeToolkit::OnOutlinerBoneSelectionChanged(UGeometryCollect
 	RootComponent->MarkRenderDynamicDataDirty();
 }
 
+void FFractureEditorModeToolkit::OnHistogramBoneSelectionChanged(UGeometryCollectionComponent* RootComponent, TArray<int32>& SelectedBones)
+{
+	// If anything is selected in the Histogram, filter the appropriate Outliner component
+	OutlinerView->SetHistogramSelection(RootComponent, SelectedBones);
+}
+
 FText FFractureEditorModeToolkit::GetStatisticsSummary() const
 {
 
diff --git a/Engine/Plugins/Experimental/ChaosEditor/Source/FractureEditor/Private/FractureToolAutoCluster.cpp b/Engine/Plugins/Experimental/ChaosEditor/Source/FractureEditor/Private/FractureToolAutoCluster.cpp
--- a/Engine/Plugins/Experimental/ChaosEditor/Source/FractureEditor/Private/FractureToolAutoCluster.cpp
+++ b/Engine/Plugins/Experimental/ChaosEditor/Source/FractureEditor/Private/FractureToolAutoCluster.cpp
@@ -4,6 +4,7 @@
 
 #include "FractureEditorStyle.h"
 #include "FractureEditorCommands.h"
+#include "FractureToolContext.h"
 
 #include "AutoClusterFracture.h"
 
@@ -60,25 +61,22 @@ void UFractureToolAutoCluster::Execute(TWeakPtr<FFractureEditorModeToolkit> InTo
 {
 	if (InToolkit.IsValid())
 	{
-		TSet<UGeometryCollectionComponent*> GeomCompSelection;
-		GetSelectedGeometryCollectionComponents(GeomCompSelection);
-		for (UGeometryCollectionComponent* GeometryCollectionComponent : GeomCompSelection)
-		{
-			FScopedColorEdit EditBoneColor = GeometryCollectionComponent->EditBoneSelection();
-			int32 CurrentLevelView = EditBoneColor.GetViewLevel();
+		FFractureEditorModeToolkit* Toolkit = InToolkit.Pin().Get();
+
+		TArray<FFractureToolContext> Contexts = GetFractureToolContexts();
 
-			UAutoClusterFractureCommand::ClusterChildBonesOfASingleMesh(GeometryCollectionComponent, AutoClusterSettings->AutoClusterMode, AutoClusterSettings->SiteCount);
+		for (FFractureToolContext& Context : Contexts)
+		{
+			Context.ConvertSelectionToClusterNodes();
 
-			EditBoneColor.ResetBoneSelection();
-			EditBoneColor.SetLevelViewMode(CurrentLevelView);
+			UAutoClusterFractureCommand::ClusterChildBonesOfASingleMesh(Context.GetGeometryCollectionComponent(), AutoClusterSettings->AutoClusterMode, AutoClusterSettings->SiteCount);
 
-			GeometryCollectionComponent->MarkRenderDynamicDataDirty();
-			GeometryCollectionComponent->MarkRenderStateDirty();
+			Refresh(Context, Toolkit);
 		}
 
-		InToolkit.Pin()->SetOutlinerComponents(GeomCompSelection.Array());
+		SetOutlinerComponents(Contexts, Toolkit);
 	}
 }
-
+	
 #undef LOCTEXT_NAMESPACE
 
diff --git a/Engine/Plugins/Experimental/ChaosEditor/Source/FractureEditor/Private/FractureToolBrick.cpp b/Engine/Plugins/Experimental/ChaosEditor/Source/FractureEditor/Private/FractureToolBrick.cpp
--- a/Engine/Plugins/Experimental/ChaosEditor/Source/FractureEditor/Private/FractureToolBrick.cpp
+++ b/Engine/Plugins/Experimental/ChaosEditor/Source/FractureEditor/Private/FractureToolBrick.cpp
@@ -11,6 +11,7 @@
 #include "GeometryCollection/GeometryCollection.h"
 #include "GeometryCollection/GeometryCollectionObject.h"
 #include "PlanarCut.h"
+#include "FractureToolContext.h"
 
 #define LOCTEXT_NAMESPACE "FractureBrick"
 
@@ -195,20 +196,14 @@ void UFractureToolBrick::GenerateBrickTransforms(const FBox& Bounds)
 
 void UFractureToolBrick::PostEditChangeChainProperty(struct FPropertyChangedChainEvent& PropertyChangedEvent)
 {
-	FFractureToolContext FractureContext;
-
-	FractureContext.Bounds = FBox(ForceInitToZero);
+	
+	FBox Bounds(ForceInitToZero);
 
 	USelection* SelectionSet = GEditor->GetSelectedActors();
 
 	TArray<AActor*> SelectedActors;
 	SelectedActors.Reserve(SelectionSet->Num());
 
-	FractureContext.RandomSeed = FMath::Rand();
-	if (CutterSettings->RandomSeed > -1)
-	{
-		FractureContext.RandomSeed = CutterSettings->RandomSeed;
-	}
 
 	FBox SelectedMeshBounds(ForceInit);
 
@@ -222,31 +217,30 @@ void UFractureToolBrick::PostEditChangeChainProperty(struct FPropertyChangedChai
 		Actor->GetComponents(PrimitiveComponents);
 		for (UPrimitiveComponent* PrimitiveComponent : PrimitiveComponents)
 		{
-			FractureContext.OriginalActor = Actor;
-			FractureContext.OriginalPrimitiveComponent = PrimitiveComponent;
-			FractureContext.Transform = Actor->GetTransform();
 			FVector Origin;
 			FVector BoxExtent;
 			Actor->GetActorBounds(false, Origin, BoxExtent);
+
 			if (CutterSettings->bGroupFracture)
 			{
-				FractureContext.Bounds += FBox::BuildAABB(Origin, BoxExtent);
+				Bounds += FBox::BuildAABB(Origin, BoxExtent);
 			}
 			else
 			{
-				FractureContext.Bounds = FBox::BuildAABB(Origin, BoxExtent);
-				GenerateBrickTransforms(FractureContext.Bounds);
+				GenerateBrickTransforms(FBox::BuildAABB(Origin, BoxExtent));
 			}
 		}
 	}
 
 	if (CutterSettings->bGroupFracture)
 	{
-		GenerateBrickTransforms(FractureContext.Bounds);
+		GenerateBrickTransforms(Bounds);
 	}
 
 }
 
+
+
 void UFractureToolBrick::Render(const FSceneView* View, FViewport* Viewport, FPrimitiveDrawInterface* PDI)
 {
 	//TODO :Plane cutting drawing
@@ -286,45 +280,43 @@ void UFractureToolBrick::AddBoxEdges(const FVector& Min, const FVector& Max)
 	Edges.Emplace(MakeTuple(FVector(Min.X, Max.Y, Max.Z), Max));
 }
 
-void UFractureToolBrick::ExecuteFracture(const FFractureToolContext& FractureContext)
+int32 UFractureToolBrick::ExecuteFracture(const FFractureToolContext& FractureContext)
 {
-	if (FractureContext.FracturedGeometryCollection != nullptr)
+	if (FractureContext.IsValid())
 	{
- 		TSharedPtr<FGeometryCollection, ESPMode::ThreadSafe> GeometryCollectionPtr = FractureContext.FracturedGeometryCollection->GetGeometryCollection();
- 		if (FGeometryCollection* GeometryCollection = GeometryCollectionPtr.Get())
- 		{
-			BrickTransforms.Empty();
-			GenerateBrickTransforms(FractureContext.Bounds);
+ 		BrickTransforms.Empty();
+		GenerateBrickTransforms(FractureContext.GetBounds());
 
-			const FQuat HeaderRotation(FVector::UpVector, 1.5708);
+		const FQuat HeaderRotation(FVector::UpVector, 1.5708);
 
-			const float HalfHeight = BrickSettings->BrickHeight * 0.5f;
-			const float HalfDepth = BrickSettings->BrickDepth * 0.5f;
-			const float HalfLength = BrickSettings->BrickLength * 0.5f;
+		const float HalfHeight = BrickSettings->BrickHeight * 0.5f;
+		const float HalfDepth = BrickSettings->BrickDepth * 0.5f;
+		const float HalfLength = BrickSettings->BrickLength * 0.5f;
 
-			FVector Max(HalfLength, HalfDepth, HalfHeight);
-			TArray<FBox> BricksToCut;
+		FVector Max(HalfLength, HalfDepth, HalfHeight);
+		TArray<FBox> BricksToCut;
 
-			for (const FTransform& Trans : BrickTransforms)
-			{
-				BricksToCut.Add(FBox(-Max * 0.95f, Max * 0.95f).TransformBy(Trans));
-			}
+		for (const FTransform& Trans : BrickTransforms)
+		{
+			BricksToCut.Add(FBox(-Max * 0.95f, Max * 0.95f).TransformBy(Trans));
+		}
 
- 			FPlanarCells VoronoiPlanarCells = FPlanarCells(BricksToCut);
+ 		FPlanarCells VoronoiPlanarCells = FPlanarCells(BricksToCut);
 
-			FNoiseSettings NoiseSettings;
-			if (CutterSettings->Amplitude > 0.0f)
-			{
-				NoiseSettings.Amplitude = CutterSettings->Amplitude;
-				NoiseSettings.Frequency = CutterSettings->Frequency;
-				NoiseSettings.Octaves = CutterSettings->OctaveNumber;
-				NoiseSettings.PointSpacing = CutterSettings->SurfaceResolution;
-				VoronoiPlanarCells.InternalSurfaceMaterials.NoiseSettings = NoiseSettings;
-			}
+		FNoiseSettings NoiseSettings;
+		if (CutterSettings->Amplitude > 0.0f)
+		{
+			NoiseSettings.Amplitude = CutterSettings->Amplitude;
+			NoiseSettings.Frequency = CutterSettings->Frequency;
+			NoiseSettings.Octaves = CutterSettings->OctaveNumber;
+			NoiseSettings.PointSpacing = CutterSettings->SurfaceResolution;
+			VoronoiPlanarCells.InternalSurfaceMaterials.NoiseSettings = NoiseSettings;
+		}
 
-			CutMultipleWithPlanarCells(VoronoiPlanarCells, *GeometryCollection, FractureContext.SelectedBones);
- 		}
+		return CutMultipleWithPlanarCells(VoronoiPlanarCells, *FractureContext.GetGeometryCollection(), FractureContext.GetSelection());
 	}
+
+	return INDEX_NONE;
 }
 
 #undef LOCTEXT_NAMESPACE
diff --git a/Engine/Plugins/Experimental/ChaosEditor/Source/FractureEditor/Private/FractureToolBrick.h b/Engine/Plugins/Experimental/ChaosEditor/Source/FractureEditor/Private/FractureToolBrick.h
--- a/Engine/Plugins/Experimental/ChaosEditor/Source/FractureEditor/Private/FractureToolBrick.h
+++ b/Engine/Plugins/Experimental/ChaosEditor/Source/FractureEditor/Private/FractureToolBrick.h
@@ -7,6 +7,7 @@
 
 #include "FractureToolBrick.generated.h"
 
+class FFractureToolContext;
 
 /** Brick Projection Directions*/
 UENUM()
@@ -93,7 +94,7 @@ public:
 	virtual TArray<UObject*> GetSettingsObjects() const override;
 	void PostEditChangeChainProperty(struct FPropertyChangedChainEvent& PropertyChangedEvent) override;
 	virtual void Render(const FSceneView* View, FViewport* Viewport, FPrimitiveDrawInterface* PDI) override;
-	virtual void ExecuteFracture(const FFractureToolContext& FractureContext) override;
+	virtual int32 ExecuteFracture(const FFractureToolContext& FractureContext) override;
 
 	void GenerateBrickTransforms(const FBox& Bounds);
 
diff --git a/Engine/Plugins/Experimental/ChaosEditor/Source/FractureEditor/Private/FractureToolClusterCutter.cpp b/Engine/Plugins/Experimental/ChaosEditor/Source/FractureEditor/Private/FractureToolClusterCutter.cpp
--- a/Engine/Plugins/Experimental/ChaosEditor/Source/FractureEditor/Private/FractureToolClusterCutter.cpp
+++ b/Engine/Plugins/Experimental/ChaosEditor/Source/FractureEditor/Private/FractureToolClusterCutter.cpp
@@ -3,7 +3,7 @@
 
 #include "FractureToolClusterCutter.h"
 
-#include "FractureEditorStyle.h"
+#include "FractureToolContext.h"
 
 #define LOCTEXT_NAMESPACE "FractureClustered"
 
@@ -44,19 +44,19 @@ TArray<UObject*> UFractureToolClusterCutter::GetSettingsObjects() const
 	return ReturnSettings;
 }
 
-void UFractureToolClusterCutter::GenerateVoronoiSites(const FFractureToolContext &Context, TArray<FVector>& Sites)
+void UFractureToolClusterCutter::GenerateVoronoiSites(const FFractureToolContext& Context, TArray<FVector>& Sites)
 {
- 	FRandomStream RandStream(Context.RandomSeed);
+ 	FRandomStream RandStream(Context.GetSeed());
 	const int32 SiteCount = RandStream.RandRange(ClusterSettings->NumberClustersMin, ClusterSettings->NumberClustersMax);
 
-	const FVector Extent(Context.Bounds.Max - Context.Bounds.Min);
+	const FVector Extent(Context.GetBounds().Max - Context.GetBounds().Min);
 
 	TArray<FVector> CenterSites;
 
 	Sites.Reserve(Sites.Num() + SiteCount);
 	for (int32 ii = 0; ii < SiteCount; ++ii)
 	{
-		CenterSites.Emplace(Context.Bounds.Min + FVector(RandStream.FRand(), RandStream.FRand(), RandStream.FRand()) * Extent);
+		CenterSites.Emplace(Context.GetBounds().Min + FVector(RandStream.FRand(), RandStream.FRand(), RandStream.FRand()) * Extent);
 	}
 
 	for (int32 kk = 0, nk = CenterSites.Num(); kk < nk; ++kk)
@@ -67,7 +67,7 @@ void UFractureToolClusterCutter::GenerateVoronoiSites(const FFractureToolContext
 		{
 			FVector V(RandStream.VRand());
 			V.Normalize();
-			V *= ClusterSettings->ClusterRadius + (RandStream.FRandRange(ClusterSettings->ClusterRadiusPercentageMin, ClusterSettings->ClusterRadiusPercentageMax) * Context.Bounds.GetExtent().GetAbsMax());
+			V *= ClusterSettings->ClusterRadius + (RandStream.FRandRange(ClusterSettings->ClusterRadiusPercentageMin, ClusterSettings->ClusterRadiusPercentageMax) * Context.GetBounds().GetExtent().GetAbsMax());
 			V += CenterSites[kk];
 			Sites.Emplace(V);
 		}
diff --git a/Engine/Plugins/Experimental/ChaosEditor/Source/FractureEditor/Private/FractureToolClusterCutter.h b/Engine/Plugins/Experimental/ChaosEditor/Source/FractureEditor/Private/FractureToolClusterCutter.h
--- a/Engine/Plugins/Experimental/ChaosEditor/Source/FractureEditor/Private/FractureToolClusterCutter.h
+++ b/Engine/Plugins/Experimental/ChaosEditor/Source/FractureEditor/Private/FractureToolClusterCutter.h
@@ -7,6 +7,7 @@
 
 #include "FractureToolClusterCutter.generated.h"
 
+class FFractureToolContext;
 
 UCLASS(config = EditorPerProjectUserSettings)
 class UFractureClusterCutterSettings : public UFractureToolSettings
@@ -75,6 +76,6 @@ public:
 	UFractureClusterCutterSettings* ClusterSettings;
 
 protected:
-	void GenerateVoronoiSites(const FFractureToolContext &Context, TArray<FVector>& Sites) override;
+	void GenerateVoronoiSites(const FFractureToolContext& Context, TArray<FVector>& Sites) override;
 
 };
\ No newline at end of file
diff --git a/Engine/Plugins/Experimental/ChaosEditor/Source/FractureEditor/Private/FractureToolClusterMagnet.cpp b/Engine/Plugins/Experimental/ChaosEditor/Source/FractureEditor/Private/FractureToolClusterMagnet.cpp
--- a/Engine/Plugins/Experimental/ChaosEditor/Source/FractureEditor/Private/FractureToolClusterMagnet.cpp
+++ b/Engine/Plugins/Experimental/ChaosEditor/Source/FractureEditor/Private/FractureToolClusterMagnet.cpp
@@ -5,6 +5,7 @@
 #include "FractureTool.h"
 #include "FractureEditorStyle.h"
 #include "FractureEditorCommands.h"
+#include "FractureToolContext.h"
 
 #include "Chaos/TriangleMesh.h"
 #include "PhysicsProxy/GeometryCollectionPhysicsProxy.h"
@@ -44,7 +45,7 @@ FText UFractureToolClusterMagnet::GetTooltipText() const
 
 FText UFractureToolClusterMagnet::GetApplyText() const
 {
-	return FText(NSLOCTEXT("Fracturet", "ExecuteClusterMagnet", "Cluster Magnet"));
+	return FText(NSLOCTEXT("Fracture", "ExecuteClusterMagnet", "Cluster Magnet"));
 }
 
 FSlateIcon UFractureToolClusterMagnet::GetToolIcon() const
@@ -62,107 +63,91 @@ TArray<UObject*> UFractureToolClusterMagnet::GetSettingsObjects() const
 
 void UFractureToolClusterMagnet::RegisterUICommand(FFractureEditorCommands* BindingContext)
 {
-	UI_COMMAND_EXT(BindingContext, UICommandInfo, "ClusterMagnet", "Cluster Magnet", "Builds clusters at local level by collecting bones adjacent to clusters or bones with highest mass.", EUserInterfaceActionType::ToggleButton, FInputChord());
+	UI_COMMAND_EXT(BindingContext, UICommandInfo, "ClusterMagnet", "Magnet", "Builds clusters at local level by collecting bones adjacent to clusters or bones with highest mass.", EUserInterfaceActionType::ToggleButton, FInputChord());
 	BindingContext->ClusterMagnet = UICommandInfo;
 }
 
+
 void UFractureToolClusterMagnet::Execute(TWeakPtr<FFractureEditorModeToolkit> InToolkit)
 {
-	
 	if (InToolkit.IsValid())
 	{
-		TSet<UGeometryCollectionComponent*> GeomCompSelection;
-		GetSelectedGeometryCollectionComponents(GeomCompSelection);
-		for (UGeometryCollectionComponent* GeometryCollectionComponent : GeomCompSelection)
+		FFractureEditorModeToolkit* Toolkit = InToolkit.Pin().Get();
+
+		int32 CurrentLevelView = Toolkit->GetLevelViewValue();
+
+		TArray<FFractureToolContext> Contexts = GetFractureToolContexts();
+
+		for (FFractureToolContext& Context : Contexts)
 		{
-			const UGeometryCollection* RestCollection = GeometryCollectionComponent->GetRestCollection();
-			if (RestCollection)
+			// We require certain attributes present to proceed.
+			if (!CheckPresenceOfNecessaryAttributes(Context.GetGeometryCollection()))
 			{
-				FGeometryCollectionPtr GeometryCollection = RestCollection->GetGeometryCollection();
+				continue;
+			}
 
-				// We require certain attributes present to proceed.
-				if (!CheckPresenceOfNecessaryAttributes(GeometryCollection))
-				{
-					return;
-				}
+			const TManagedArray<TSet<int32>>& Children = Context.GetGeometryCollection()->Children;
+			const TManagedArray<int32>& Levels = Context.GetGeometryCollection()->GetAttribute<int32>("Level", FGeometryCollection::TransformGroup);
+
+			Context.Sanitize();
+			TMap<int32, TArray<int32>> ClusteredSelection = Context.GetClusteredSelections();
 
-				// If no bones are selected, assume that we're working on the root's children
-				TArray<int32> SelectedBones = GeometryCollectionComponent->GetSelectedBones();
-				if (SelectedBones.Num() == 0)
+			for (TPair<int32, TArray<int32>>& Group : ClusteredSelection)
+			{
+				const TArray<int32>& CurrentSelection = Group.Value;
+				CurrentLevelView = Levels[Group.Key] + 1;
+
+				// We have the connections for the leaf nodes of our geometry collection. We want to percolate those up to the top nodes.
+				TMap<int32, TSet<int32>> TopNodeConnectivity = InitializeConnectivity(Children[Group.Key], Context.GetGeometryCollection(), Levels[Group.Key]+1);
+
+				// Separate the top nodes into cluster magnets and a pool of available nodes.
+				TArray<FClusterMagnet> ClusterMagnets;
+				TSet<int32> RemainingPool;
+				SeparateClusterMagnets(Children[Group.Key], Group.Value, TopNodeConnectivity, ClusterMagnets, RemainingPool);
+
+				for (uint32 Iteration = 0; Iteration < ClusterMagnetSettings->Iterations; ++Iteration)
 				{
-					FGeometryCollectionClusteringUtility::GetRootBones(GeometryCollection.Get(), SelectedBones);
+					bool bNeighborsAbsorbed = false;
+
+					// each cluster gathers adjacent nodes from the pool
+					for (FClusterMagnet& ClusterMagnet : ClusterMagnets)
+					{
+						bNeighborsAbsorbed |= AbsorbClusterNeighbors(TopNodeConnectivity, ClusterMagnet, RemainingPool);
+					}
+
+					// early termination
+					if (!bNeighborsAbsorbed)
+					{
+						break;
+					}
 				}
 
-				for (int32 CurrentRoot : SelectedBones)
+				// Create new clusters from the cluster magnets
+				for (const FClusterMagnet& ClusterMagnet : ClusterMagnets)
 				{
-					const TManagedArray<TSet<int32>>& Children = GeometryCollection->Children;
-					if (Children[CurrentRoot].Num() > 0)
+					if (ClusterMagnet.ClusteredNodes.Num() > 1)
 					{
-						const TManagedArray<int32>& Levels = GeometryCollection->GetAttribute<int32>("Level", FTransformCollection::TransformGroup);
-						const int32 OperatingLevel = Levels[CurrentRoot] + 1;
-						//const TArray<int32> TopNodes = GatherTopNodes(GeometryCollection, CurrentRoot);
-						const TSet<int32> TopNodes = Children[CurrentRoot];
-
-						UpdateMasses(GeometryCollection, TopNodes);
-
-						float CutoffMass = FindCutoffMass(ClusterMagnetSettings->MassPercentile, GeometryCollection, TopNodes);
-
-						// We have the connections for the leaf nodes of our geometry collection. We want to percolate those up to the top nodes.
-						TMap<int32, TSet<int32>> TopNodeConnectivity = InitializeConnectivity(TopNodes, GeometryCollection, OperatingLevel);
-
-						// Separate the top nodes into cluster magnets and a pool of available nodes.
-						TArray<FClusterMagnet> ClusterMagnets;
-						TSet<int32> RemainingPool;
-						SeparateClusterMagnets(GeometryCollection, TopNodes, CutoffMass, TopNodeConnectivity, ClusterMagnets, RemainingPool);
-
-						for (uint32 Iteration = 0; Iteration < ClusterMagnetSettings->Iterations; ++Iteration)
-						{
-							bool bNeighborsAbsorbed = false;
-
-							// each cluster gathers adjacent nodes from the pool
-							for (FClusterMagnet& ClusterMagnet : ClusterMagnets)
-							{
-								bNeighborsAbsorbed |= AbsorbClusterNeighbors(TopNodeConnectivity, ClusterMagnet, RemainingPool);
-							}
-
-							// early termination
-							if (!bNeighborsAbsorbed)
-							{
-								break;
-							}
-						}
-
-						// Create new clusters from the cluster magnets
-						for (const FClusterMagnet& ClusterMagnet : ClusterMagnets)
-						{
-							if (ClusterMagnet.ClusteredNodes.Num() > 1)
-							{
-								TArray<int32> NewChildren = ClusterMagnet.ClusteredNodes.Array();
-								NewChildren.Sort();
-								FGeometryCollectionClusteringUtility::ClusterBonesUnderNewNode(GeometryCollection.Get(), NewChildren[0], NewChildren, false, false);
-								
-							}
-						}
-
-						FGeometryCollectionClusteringUtility::UpdateHierarchyLevelOfChildren(GeometryCollection.Get(), CurrentRoot);
+						TArray<int32> NewChildren = ClusterMagnet.ClusteredNodes.Array();
+						NewChildren.Sort();
+						FGeometryCollectionClusteringUtility::ClusterBonesUnderNewNode(Context.GetGeometryCollection().Get(), NewChildren[0], NewChildren, false, false);
+
 					}
 				}
+			}
 
-				FScopedColorEdit EditBoneColor = GeometryCollectionComponent->EditBoneSelection();
-				EditBoneColor.ResetBoneSelection();
-				EditBoneColor.ResetHighlightedBones();
-				InToolkit.Pin()->SetBoneSelection(GeometryCollectionComponent, EditBoneColor.GetSelectedBones(), true);
+			Refresh(Context, Toolkit);
+		}
 
-				GeometryCollectionComponent->MarkRenderDynamicDataDirty();
-				GeometryCollectionComponent->MarkRenderStateDirty();				
-			}
+		// Set the current viewing level to our best guess
+		if (CurrentLevelView != Toolkit->GetLevelViewValue())
+		{
+			Toolkit->OnSetLevelViewValue(CurrentLevelView);
 		}
 
-		InToolkit.Pin()->SetOutlinerComponents(GeomCompSelection.Array());
+		SetOutlinerComponents(Contexts, Toolkit);
 	}
 }
 
-
 bool UFractureToolClusterMagnet::CheckPresenceOfNecessaryAttributes(const FGeometryCollectionPtr GeometryCollection) const
 {
 	if (!GeometryCollection->HasAttribute("Level", FTransformCollection::TransformGroup))
@@ -180,97 +165,6 @@ bool UFractureToolClusterMagnet::CheckPresenceOfNecessaryAttributes(const FGeome
 	return true;
 }
 
-float UFractureToolClusterMagnet::FindCutoffMass(float Percentile, const FGeometryCollectionPtr GeometryCollection, const TSet<int32>& TopNodes) const
-{
-	check(Percentile >= 0.0);
-	check(Percentile <= 1.0);
-	
-	const TManagedArray<float>& Mass = GeometryCollection->GetAttribute<float>("Mass", FTransformCollection::TransformGroup);
-
-	// Collect the top node masses, sort by mass, return the threshold mass at the cutoff threshold. 
-	TArray<float> Masses;
-	Masses.Reserve(TopNodes.Num());
-	for (int32 Index : TopNodes)
-	{
-		Masses.Add(Mass[Index]);
-	}
-	Masses.Sort();
-
-	int32 ThresholdIndex = FMath::Floor(TopNodes.Num() * Percentile);
-	return Masses[ThresholdIndex];
-}
-
-
-void UFractureToolClusterMagnet::UpdateMasses(FGeometryCollectionPtr GeometryCollection, const TSet<int32>& TopNodes) const 
-{
-	if (!GeometryCollection->HasAttribute("Mass", FTransformCollection::TransformGroup))
-	{
-		GeometryCollection->AddAttribute<Chaos::FReal>("Mass", FTransformCollection::TransformGroup);
-		UE_LOG(LogFractureTool, Warning, TEXT("Added Mass attribute needed to execute ClusterMagnet."));
-	}
-	
-	TArray<FTransform> Transform;
-	GeometryCollectionAlgo::GlobalMatrices(GeometryCollection->Transform, GeometryCollection->Parent, Transform);
-
-	const TManagedArray<FVector>& Vertex = GeometryCollection->Vertex;
-	const TManagedArray<int32>& BoneMap = GeometryCollection->BoneMap;
-	
-	Chaos::TParticles<float, 3> MassSpaceParticles;
-	MassSpaceParticles.AddParticles(Vertex.Num());
-	for (int32 Idx = 0; Idx < Vertex.Num(); ++Idx)
-	{
-		MassSpaceParticles.X(Idx) = Transform[BoneMap[Idx]].TransformPosition(Vertex[Idx]);
-	}
-
-	for (int32 Index : TopNodes)
-	{
-		UpdateMasses(GeometryCollection, MassSpaceParticles, Index);
-	}
-}
-
-void UFractureToolClusterMagnet::UpdateMasses(FGeometryCollectionPtr GeometryCollection, const Chaos::TParticles<float, 3>& MassSpaceParticles, int32 TransformIndex) const
-{
-	const TManagedArray<TSet<int32>>& Children = GeometryCollection->Children;
-	const TManagedArray<bool>& Visible = GeometryCollection->Visible;
-	const TManagedArray<int32>& FaceCount = GeometryCollection->FaceCount;
-	const TManagedArray<int32>& FaceStart = GeometryCollection->FaceStart;
-	const TManagedArray<int32>& TransformToGeometryIndex = GeometryCollection->TransformToGeometryIndex;
-	const TManagedArray<FIntVector>& Indices = GeometryCollection->Indices;
-	TManagedArray<float>& Mass = GeometryCollection->GetAttribute<float>("Mass", FTransformCollection::TransformGroup);
-	
-	if (Children[TransformIndex].Num() == 0) // leaf node
-	{
-		int32 GeometryIndex = TransformToGeometryIndex[TransformIndex];
-
-		TUniquePtr<Chaos::TTriangleMesh<float>> TriMesh(
-			CreateTriangleMesh(
-				FaceStart[GeometryIndex],
-				FaceCount[GeometryIndex],
-				Visible,
-				Indices,
-				true));
-		
-		float Volume = 0.0;
-		Chaos::TVector<Chaos::FReal,3> CenterOfMass;
-		Chaos::CalculateVolumeAndCenterOfMass(MassSpaceParticles, TriMesh->GetElements(), Volume, CenterOfMass);
-
-		// Since we're only interested in relative mass, we assume density = 1.0
-		Mass[TransformIndex] = Volume;
-	}
-	else
-	{
-		// Recurse to children and sum the masses for this node
-		float LocalMass = 0.0;
-		for (int32 ChildIndex : Children[TransformIndex])
-		{
-			UpdateMasses(GeometryCollection, MassSpaceParticles, ChildIndex);
-			LocalMass += Mass[ChildIndex];
-		}
-
-		Mass[TransformIndex] = LocalMass;
-	}
-}
-
 TMap<int32, TSet<int32>> UFractureToolClusterMagnet::InitializeConnectivity(const TSet<int32>& TopNodes, FGeometryCollectionPtr GeometryCollection, int32 OperatingLevel) const
 {
 	FGeometryCollectionProximityUtility::UpdateProximity(GeometryCollection.Get());
@@ -318,23 +212,20 @@ void UFractureToolClusterMagnet::CollectTopNodeConnections(FGeometryCollectionPt
 }
 
 void UFractureToolClusterMagnet::SeparateClusterMagnets(
-	const FGeometryCollectionPtr GeometryCollection,
 	const TSet<int32>& TopNodes,
-	float CutoffMass,
+	const TArray<int32>& Selection,
 	const TMap<int32, TSet<int32>>& TopNodeConnectivity,
 	TArray<FClusterMagnet>& OutClusterMagnets,
 	TSet<int32>& OutRemainingPool) const
 {
 	// Push any top nodes over the mass threshold into cluster magnets.
 	
-	const TManagedArray<float>& Mass = GeometryCollection->GetAttribute<float>("Mass", FTransformCollection::TransformGroup);
-	
 	OutClusterMagnets.Reserve(TopNodes.Num());
 	OutRemainingPool.Reserve(TopNodes.Num());
-
+	
 	for (int32 Index : TopNodes)
 	{
-		if (Mass[Index] > CutoffMass)
+		if (Selection.Contains(Index))
 		{
 			OutClusterMagnets.AddDefaulted();
 			FClusterMagnet& NewMagnet = OutClusterMagnets.Last();
diff --git a/Engine/Plugins/Experimental/ChaosEditor/Source/FractureEditor/Private/FractureToolClusterMagnet.h b/Engine/Plugins/Experimental/ChaosEditor/Source/FractureEditor/Private/FractureToolClusterMagnet.h
--- a/Engine/Plugins/Experimental/ChaosEditor/Source/FractureEditor/Private/FractureToolClusterMagnet.h
+++ b/Engine/Plugins/Experimental/ChaosEditor/Source/FractureEditor/Private/FractureToolClusterMagnet.h
@@ -29,15 +29,11 @@ public:
 		UFractureClusterMagnetSettings(const FObjectInitializer& ObjInit)
 		: Super(ObjInit)
 		, Iterations(1)
-		, MassPercentile(0.9)
 	{}
 
 	UPROPERTY(EditAnywhere, Category = ClusterMagnet, meta = (DisplayName = "Iterations", UIMin = "1", UIMax = "100", ClampMin = "1", ClampMax = "500"))
 	uint32 Iterations;
 
-	UPROPERTY(EditAnywhere, Category = ClusterMagnet, meta = (DisplayName = "Mass percentile to form magnets", UIMin = "0.0", UIMax = "1.0", ClampMin = "0.0", ClampMax = "1.0"))
-	float MassPercentile;
-
 };
 
 
@@ -66,13 +62,9 @@ public:
 
 private:
 	bool CheckPresenceOfNecessaryAttributes(const FGeometryCollectionPtr GeometryCollection) const;
-	//TArray<int32> GatherTopNodes(const FGeometryCollectionPtr GeometryCollection) const;
-	float FindCutoffMass(float Percentile, const FGeometryCollectionPtr GeometryCollection, const TSet<int32>& TopNodes) const;
-	void UpdateMasses(FGeometryCollectionPtr GeometryCollection, const TSet<int32>& TopNodes) const;
-	void UpdateMasses(FGeometryCollectionPtr GeometryCollection, const Chaos::TParticles<float, 3>& MassSpaceParticles, int32 TransformIndex) const;
 	TMap<int32, TSet<int32>> InitializeConnectivity(const TSet<int32>& TopNodes, FGeometryCollectionPtr GeometryCollection, int32 OperatingLevel) const;
 	void CollectTopNodeConnections(FGeometryCollectionPtr GeometryCollection, int32 Index, int32 OperatingLevel, TSet<int32>& OutConnections) const;
-	void SeparateClusterMagnets(const FGeometryCollectionPtr GeometryCollection, const TSet<int32>& TopNodes, float CutoffMass, const TMap<int32, TSet<int32>>& TopNodeConnectivity, TArray<FClusterMagnet>& OutClusterMagnets, TSet<int32>& OutRemainingPool) const;
+	void SeparateClusterMagnets(const TSet<int32>& TopNodes, const TArray<int32>& Selection, const TMap<int32, TSet<int32>>& TopNodeConnectivity, TArray<FClusterMagnet>& OutClusterMagnets, TSet<int32>& OutRemainingPool) const;
 	bool AbsorbClusterNeighbors(const TMap<int32, TSet<int32>> TopNodeConnectivity, FClusterMagnet& OutClusterMagnets, TSet<int32>& OutRemainingPool) const;
 
 };
diff --git a/Engine/Plugins/Experimental/ChaosEditor/Source/FractureEditor/Private/FractureToolClustering.cpp b/Engine/Plugins/Experimental/ChaosEditor/Source/FractureEditor/Private/FractureToolClustering.cpp
--- a/Engine/Plugins/Experimental/ChaosEditor/Source/FractureEditor/Private/FractureToolClustering.cpp
+++ b/Engine/Plugins/Experimental/ChaosEditor/Source/FractureEditor/Private/FractureToolClustering.cpp
@@ -2,6 +2,8 @@
 
 #include "FractureToolClustering.h"
 
+#include "FractureToolContext.h"
+
 #include "GeometryCollection/GeometryCollectionComponent.h"
 #include "GeometryCollection/GeometryCollection.h"
 #include "GeometryCollection/GeometryCollectionClusteringUtility.h"
@@ -34,52 +36,36 @@ void UFractureToolFlattenAll::Execute(TWeakPtr<FFractureEditorModeToolkit> InToo
 {
 	if (InToolkit.IsValid())
 	{
-		TSet<UGeometryCollectionComponent*> GeomCompSelection;
-		GetSelectedGeometryCollectionComponents(GeomCompSelection);
-		for (UGeometryCollectionComponent* GeometryCollectionComponent : GeomCompSelection)
-		{
-			FGeometryCollectionEdit GCEdit = GeometryCollectionComponent->EditRestCollection();
-			if (UGeometryCollection* GCObject = GCEdit.GetRestCollection())
-			{
-				TSharedPtr<FGeometryCollection, ESPMode::ThreadSafe> GeometryCollectionPtr = GCObject->GetGeometryCollection();
-				if (FGeometryCollection* GeometryCollection = GeometryCollectionPtr.Get())
-				{
-					AddAdditionalAttributesIfRequired(GCObject);
+		FFractureEditorModeToolkit* Toolkit = InToolkit.Pin().Get();
 
-					AddSingleRootNodeIfRequired(GCObject);
+		int32 CurrentLevelView = Toolkit->GetLevelViewValue();
 
-					int32 NumElements = GCObject->NumElements(FGeometryCollection::TransformGroup);
-					TArray<int32> Elements;
-					Elements.Reserve(NumElements);
-
-					for (int32 Element = 0; Element < NumElements; ++Element)
-					{
-						if (GeometryCollection->Parent[Element] != FGeometryCollection::Invalid)
-						{
-							Elements.Add(Element);
-						}
-					}
+		TArray<FFractureToolContext> Contexts = GetFractureToolContexts();
 
-					if (Elements.Num() > 0)
-					{
-						FGeometryCollectionClusteringUtility::ClusterBonesUnderExistingRoot(GeometryCollection, Elements);
-					}
+		for (FFractureToolContext& Context : Contexts)
+		{
+			const TManagedArray<int32>& Levels = Context.GetGeometryCollection()->GetAttribute<int32>("Level", FGeometryCollection::TransformGroup);
 
-					FGeometryCollectionClusteringUtility::UpdateHierarchyLevelOfChildren(GeometryCollection, -1);
+			Context.ConvertSelectionToClusterNodes();
 
-					FScopedColorEdit EditBoneColor = GeometryCollectionComponent->EditBoneSelection();
-					EditBoneColor.ResetBoneSelection();
+			for (int32 ClusterIndex : Context.GetSelection())
+			{
+				TArray<int32> LeafBones;
+				FGeometryCollectionClusteringUtility::GetLeafBones(Context.GetGeometryCollection().Get(), ClusterIndex, LeafBones);
+				FGeometryCollectionClusteringUtility::ClusterBonesUnderExistingNode(Context.GetGeometryCollection().Get(), ClusterIndex, LeafBones);
+				CurrentLevelView = Levels[ClusterIndex] + 1;
+			}
 
-					InToolkit.Pin()->OnSetLevelViewValue(1);
+			Refresh(Context, Toolkit);
+		}
 
-					GeometryCollectionComponent->MarkRenderDynamicDataDirty();
-					GeometryCollectionComponent->MarkRenderStateDirty();
-				}
-			}
+		if (CurrentLevelView != Toolkit->GetLevelViewValue())
+		{
+			Toolkit->OnSetLevelViewValue(CurrentLevelView);
 		}
 
-		InToolkit.Pin()->SetOutlinerComponents(GeomCompSelection.Array());
-	}	
+		SetOutlinerComponents(Contexts, Toolkit);
+	}
 }
 
 
@@ -107,59 +93,60 @@ void UFractureToolCluster::RegisterUICommand(FFractureEditorCommands* BindingCon
 
 void UFractureToolCluster::Execute(TWeakPtr<FFractureEditorModeToolkit> InToolkit)
 {
+	
 	if (InToolkit.IsValid())
 	{
-		TSet<UGeometryCollectionComponent*> GeomCompSelection;
-		GetSelectedGeometryCollectionComponents(GeomCompSelection);
-		for (UGeometryCollectionComponent* GeometryCollectionComponent : GeomCompSelection)
+		FFractureEditorModeToolkit* Toolkit = InToolkit.Pin().Get();
+
+		int32 CurrentLevelView = Toolkit->GetLevelViewValue();
+
+		TArray<FFractureToolContext> Contexts = GetFractureToolContexts();
+
+		for (FFractureToolContext& Context : Contexts)
 		{
-			const TArray<int32> SelectedBones = GeometryCollectionComponent->GetSelectedBones();
+			const TManagedArray<TSet<int32>>& Children = Context.GetGeometryCollection()->GetAttribute<TSet<int32>>("Children", FGeometryCollection::TransformGroup);
+			const TManagedArray<int32>& Levels = Context.GetGeometryCollection()->GetAttribute<int32>("Level", FGeometryCollection::TransformGroup);
+			
+			Context.Sanitize();
+			TMap<int32, TArray<int32>> ClusteredSelection = Context.GetClusteredSelections();
 
-			if (SelectedBones.Num() > 1)
+			for (TPair<int32, TArray<int32>>& Group : ClusteredSelection)
 			{
-				FGeometryCollectionEdit GCEdit = GeometryCollectionComponent->EditRestCollection();
-				if (UGeometryCollection* GCObject = GCEdit.GetRestCollection())
+				const TArray<int32>& CurrentSelection = Group.Value;
+
+				// sort the selection list so ClusterBonesUnderNewNode() happens in the correct order for leaf nodes
+				TArray<int32> SortedSelectedBones;
+				SortedSelectedBones.Reserve(CurrentSelection.Num());
+				for (int32 SelectedBone : CurrentSelection)
 				{
-					TSharedPtr<FGeometryCollection, ESPMode::ThreadSafe> GeometryCollectionPtr = GCObject->GetGeometryCollection();
-					if (FGeometryCollection* GeometryCollection = GeometryCollectionPtr.Get())
+					if (Children[SelectedBone].Num() > 0)
+					{
+						SortedSelectedBones.Insert(SelectedBone, 0);
+					}
+					else
 					{
-						const TManagedArray<TSet<int32>>& Children = GeometryCollection->GetAttribute<TSet<int32>>("Children", FGeometryCollection::TransformGroup);
-
-						// sort the selection list so ClusterBonesUnderNewNode() happens in the correct order for leaf nodes
-						TArray<int32> SortedSelectedBones;
-						SortedSelectedBones.Reserve(SelectedBones.Num());
-						for (int32 SelectedBone : SelectedBones)
-						{
-							if (Children[SelectedBone].Num() > 0)
-							{
-								SortedSelectedBones.Insert(SelectedBone, 0);
-							}
-							else
-							{
-								SortedSelectedBones.Add(SelectedBone);
-							}
-						}
-						// cluster Selected Bones under the first selected bone
-						int32 InsertAtIndex = SortedSelectedBones[0];
-
-						FGeometryCollectionClusteringUtility::ClusterBonesUnderNewNode(GeometryCollection, InsertAtIndex, SortedSelectedBones, false);
-						FGeometryCollectionClusteringUtility::UpdateHierarchyLevelOfChildren(GeometryCollection, -1);
-
-						FScopedColorEdit EditBoneColor = GeometryCollectionComponent->EditBoneSelection();
-						EditBoneColor.ResetBoneSelection();
-						EditBoneColor.ResetHighlightedBones();
-						GeometryCollectionComponent->MarkRenderDynamicDataDirty();
-						GeometryCollectionComponent->MarkRenderStateDirty();
-						InToolkit.Pin()->SetBoneSelection(GeometryCollectionComponent, EditBoneColor.GetSelectedBones(), true);
+						SortedSelectedBones.Add(SelectedBone);
 					}
 				}
+
+				// cluster Selected Bones under the first selected bone
+				int32 InsertAtIndex = SortedSelectedBones[0];
+				CurrentLevelView = Levels[InsertAtIndex];
+
+				FGeometryCollectionClusteringUtility::ClusterBonesUnderNewNode(Context.GetGeometryCollection().Get(), InsertAtIndex, SortedSelectedBones, false);
 			}
+
+			Refresh(Context, Toolkit);
 		}
 
-		InToolkit.Pin()->SetOutlinerComponents(GeomCompSelection.Array());
-	}
-}
+		if (CurrentLevelView != Toolkit->GetLevelViewValue())
+		{
+			Toolkit->OnSetLevelViewValue(CurrentLevelView);
+		}
 
+		SetOutlinerComponents(Contexts, Toolkit);
+	}	
+}
 
 
 FText UFractureToolUncluster::GetDisplayText() const
@@ -187,37 +174,48 @@ void UFractureToolUncluster::Execute(TWeakPtr<FFractureEditorModeToolkit> InTool
 {
 	if (InToolkit.IsValid())
 	{
-		int32 FractureLevel = InToolkit.Pin()->GetLevelViewValue();
-		TSet<UGeometryCollectionComponent*> GeomCompSelection;
-		GetSelectedGeometryCollectionComponents(GeomCompSelection);
-		for (UGeometryCollectionComponent* GeometryCollectionComponent : GeomCompSelection)
+		FFractureEditorModeToolkit* Toolkit = InToolkit.Pin().Get();
+
+		int32 CurrentLevelView = Toolkit->GetLevelViewValue();
+
+		TArray<FFractureToolContext> Contexts = GetFractureToolContexts();
+
+		for (FFractureToolContext& Context : Contexts)
 		{
-			// scoped edit of collection
-			FGeometryCollectionEdit GeometryCollectionEdit = GeometryCollectionComponent->EditRestCollection();
-			if (UGeometryCollection* GeometryCollectionObject = GeometryCollectionEdit.GetRestCollection())
+			const TManagedArray<TSet<int32>>& Children = Context.GetGeometryCollection()->GetAttribute<TSet<int32>>("Children", FGeometryCollection::TransformGroup);
+			const TManagedArray<int32>& Levels = Context.GetGeometryCollection()->GetAttribute<int32>("Level", FGeometryCollection::TransformGroup);
+
+			Context.ConvertSelectionToClusterNodes();
+			Context.RemoveRootNodes();
+
+			// Once the operation is complete, we'll select the children that were re-leveled
+			TArray<int32> NewSelection;
+			for (int32 Cluster : Context.GetSelection())
 			{
-				TSharedPtr<FGeometryCollection, ESPMode::ThreadSafe> GeometryCollectionPtr = GeometryCollectionObject->GetGeometryCollection();
-				if (FGeometryCollection* GeometryCollection = GeometryCollectionPtr.Get())
-				{
-					FGeometryCollectionClusteringUtility::UpdateHierarchyLevelOfChildren(GeometryCollection, -1);
-					FGeometryCollectionClusteringUtility::CollapseSelectedHierarchy(FractureLevel, GeometryCollectionComponent->GetSelectedBones(), GeometryCollection);
-
-					FScopedColorEdit EditBoneColor = GeometryCollectionComponent->EditBoneSelection();
-					EditBoneColor.ResetBoneSelection();
-					EditBoneColor.ResetHighlightedBones();
-					GeometryCollectionComponent->MarkRenderDynamicDataDirty();
-					GeometryCollectionComponent->MarkRenderStateDirty();
-					InToolkit.Pin()->SetBoneSelection(GeometryCollectionComponent, EditBoneColor.GetSelectedBones(), true);
-				}
+				NewSelection.Append(Children[Cluster].Array());
+			}
+
+			FGeometryCollectionClusteringUtility::CollapseHierarchyOneLevel(Context.GetGeometryCollection().Get(), Context.GetSelection());
+			Context.SetSelection(NewSelection);
+			if (NewSelection.Num() > 0)
+			{
+				CurrentLevelView = Levels[NewSelection.Last()];
 			}
+			
+			Refresh(Context, Toolkit);
+		}
+
+		if (CurrentLevelView != Toolkit->GetLevelViewValue())
+		{
+			Toolkit->OnSetLevelViewValue(CurrentLevelView);
 		}
-		InToolkit.Pin()->SetOutlinerComponents(GeomCompSelection.Array());
+
+		SetOutlinerComponents(Contexts, Toolkit);
 	}
 }
 
 
 
-
 FText UFractureToolMoveUp::GetDisplayText() const
 {
 	return FText(NSLOCTEXT("FractureToolClusteringOps", "FractureToolMoveUp", "Level Up"));
@@ -243,26 +241,18 @@ void UFractureToolMoveUp::Execute(TWeakPtr<FFractureEditorModeToolkit> InToolkit
 {
 	if (InToolkit.IsValid())
 	{
-		TSet<UGeometryCollectionComponent*> GeomCompSelection;
-		GetSelectedGeometryCollectionComponents(GeomCompSelection);
-		for (UGeometryCollectionComponent* GeometryCollectionComponent : GeomCompSelection)
-		{
-			FGeometryCollectionEdit GeometryCollectionEdit = GeometryCollectionComponent->EditRestCollection();
-			if (UGeometryCollection* GeometryCollectionObject = GeometryCollectionEdit.GetRestCollection())
-			{
-				TSharedPtr<FGeometryCollection, ESPMode::ThreadSafe> GeometryCollectionPtr = GeometryCollectionObject->GetGeometryCollection();
-				if (FGeometryCollection* GeometryCollection = GeometryCollectionPtr.Get())
-				{
-					TArray<int32> Selected = GeometryCollectionComponent->GetSelectedBones();
-					FGeometryCollectionClusteringUtility::MoveUpOneHierarchyLevel(GeometryCollection, Selected);
+		FFractureEditorModeToolkit* Toolkit = InToolkit.Pin().Get();
 
-					GeometryCollectionComponent->MarkRenderDynamicDataDirty();
-					GeometryCollectionComponent->MarkRenderStateDirty();
-				}
-			}
-		}
+		TArray<FFractureToolContext> Contexts = GetFractureToolContexts();
 
-		InToolkit.Pin()->SetOutlinerComponents(GeomCompSelection.Array());
+		for (FFractureToolContext& Context: Contexts)
+		{
+			Context.Sanitize();
+			FGeometryCollectionClusteringUtility::MoveUpOneHierarchyLevel(Context.GetGeometryCollection().Get(), Context.GetSelection());
+			Refresh(Context, Toolkit);
+		}
+		
+		SetOutlinerComponents(Contexts, Toolkit);
 	}
 }
 
diff --git a/Engine/Plugins/Experimental/ChaosEditor/Source/FractureEditor/Private/FractureToolContext.cpp b/Engine/Plugins/Experimental/ChaosEditor/Source/FractureEditor/Private/FractureToolContext.cpp
new file mode 100644
--- /dev/null
+++ b/Engine/Plugins/Experimental/ChaosEditor/Source/FractureEditor/Private/FractureToolContext.cpp
@@ -0,0 +1,130 @@
+// Copyright Epic Games, Inc. All Rights Reserved.
+
+#include "FractureToolContext.h"
+
+#include "GeometryCollection/GeometryCollectionClusteringUtility.h"
+
+
+FFractureToolContext::FFractureToolContext(UGeometryCollectionComponent* InGeometryCollectionComponent)
+{
+	GeometryCollectionComponent = InGeometryCollectionComponent;
+	FGeometryCollectionEdit RestCollection = GeometryCollectionComponent->EditRestCollection(GeometryCollection::EEditUpdate::None);
+	FracturedGeometryCollection = RestCollection.GetRestCollection();
+
+	Transform = GeometryCollectionComponent->GetOwner()->GetActorTransform();
+
+	FGeometryCollectionEdit GeometryCollectionEdit = GeometryCollectionComponent->EditRestCollection();
+	if (UGeometryCollection* GeometryCollectionObject = GeometryCollectionEdit.GetRestCollection())
+	{
+		GeometryCollection = GeometryCollectionObject->GetGeometryCollection();
+		SelectedBones = GeometryCollectionComponent->GetSelectedBones();
+	}
+}
+
+
+void FFractureToolContext::Sanitize()
+{
+	// Ensure that children of a selected node are not also selected.
+	SelectedBones.RemoveAll([this](int32 Index) {
+		return HasSelectedAncestor(Index);
+		});
+
+	SelectedBones.Sort();
+}
+
+void FFractureToolContext::RemoveRootNodes()
+{
+	SelectedBones.RemoveAll([this](int32 Index) {
+		return FGeometryCollectionClusteringUtility::IsARootBone(this->GeometryCollection.Get(), Index);
+		});
+}
+
+TMap<int32, TArray<int32>> FFractureToolContext::GetClusteredSelections()
+{
+	TMap<int32, TArray<int32>> SiblingGroups;
+
+	// Bin the selection indices by parent index
+	const TManagedArray<int32>& Parents = GeometryCollection->GetAttribute<int32>("Parent", FGeometryCollection::TransformGroup);
+	for (int32 Index : SelectedBones)
+	{
+		TArray<int32>& SiblingIndices = SiblingGroups.FindOrAdd(Parents[Index]);
+		SiblingIndices.Add(Index);
+	}
+
+	return SiblingGroups;
+}
+
+
+void FFractureToolContext::ConvertSelectionToLeafNodes()
+{
+	Sanitize();
+
+	TArray<int32> LeafSelection;
+	for (int32 Index : SelectedBones)
+	{
+		ConvertSelectionToLeafNodes(Index, LeafSelection);
+	}
+
+	SelectedBones = LeafSelection;
+}
+
+void FFractureToolContext::ConvertSelectionToLeafNodes(int32 Index, TArray<int32>& LeafSelection)
+{
+	const TManagedArray<TSet<int32>>& Children = GeometryCollection->GetAttribute<TSet<int32>>("Children", FGeometryCollection::TransformGroup);
+	if (Children[Index].Num() == 0)
+	{
+		LeafSelection.Add(Index);
+	}
+	else 
+	{
+		for (int32 Child : Children[Index])
+		{
+			ConvertSelectionToLeafNodes(Child, LeafSelection);
+		}
+	}	
+}
+
+
+void FFractureToolContext::ConvertSelectionToClusterNodes()
+{
+	// If this is a leaf node, select the cluster instead.
+	const TManagedArray<TSet<int32>>& Children = GeometryCollection->GetAttribute<TSet<int32>>("Children", FGeometryCollection::TransformGroup);
+	const TManagedArray<int32>& Parents = GeometryCollection->GetAttribute<int32>("Parent", FGeometryCollection::TransformGroup);
+
+	TArray<int32> AddedClusterSelections;
+	for (int32 Index : SelectedBones)
+	{
+		if (Children[Index].Num() == 0)
+		{
+			AddedClusterSelections.AddUnique(Parents[Index]);
+		}
+	}
+	SelectedBones.Append(AddedClusterSelections);
+	
+	Sanitize();
+}
+
+void FFractureToolContext::TransformBoundsToWorld()
+{
+	Bounds = Bounds.TransformBy(Transform);
+}
+
+
+bool FFractureToolContext::HasSelectedAncestor(int32 Index) const
+{
+	const TManagedArray<int32>& Parents = GeometryCollection->GetAttribute<int32>("Parent", FGeometryCollection::TransformGroup);
+
+	int32 CurrIndex = Index;
+	while (Parents[CurrIndex] != INDEX_NONE)
+	{
+		CurrIndex = Parents[CurrIndex];
+		if (SelectedBones.Contains(CurrIndex))
+		{
+			return true;
+		}
+	}
+
+	// We've arrived at the top of the hierarchy with no selected ancestors
+	return false;
+}
+
diff --git a/Engine/Plugins/Experimental/ChaosEditor/Source/FractureEditor/Private/FractureToolContext.h b/Engine/Plugins/Experimental/ChaosEditor/Source/FractureEditor/Private/FractureToolContext.h
new file mode 100644
--- /dev/null
+++ b/Engine/Plugins/Experimental/ChaosEditor/Source/FractureEditor/Private/FractureToolContext.h
@@ -0,0 +1,60 @@
+// Copyright Epic Games, Inc. All Rights Reserved.
+
+#pragma once
+
+#include "CoreMinimal.h"
+
+#include "GeometryCollection/GeometryCollection.h"
+#include "GeometryCollection/GeometryCollectionComponent.h"
+#include "GeometryCollection/GeometryCollectionObject.h"
+
+
+class FFractureToolContext
+{
+public:
+
+	using FGeometryCollectionPtr = TSharedPtr<FGeometryCollection, ESPMode::ThreadSafe>;
+
+	/** Generate a tool context based on the bone selection of the specified GeometryCollectionComponent. */
+	FFractureToolContext(UGeometryCollectionComponent* InGeometryCollectionComponent);
+
+	void Sanitize();
+	void ConvertSelectionToLeafNodes();
+	void ConvertSelectionToClusterNodes();
+	void RemoveRootNodes();
+	
+	const TArray<int32>& GetSelection() const { return SelectedBones; }
+	TArray<int32>& GetSelection() { return SelectedBones; }
+	void SetSelection(const TArray<int32>& NewSelection) { SelectedBones = NewSelection; }
+	UGeometryCollectionComponent* GetGeometryCollectionComponent() const { return GeometryCollectionComponent; }
+	FGeometryCollectionPtr GetGeometryCollection() const { return GeometryCollection; }
+	UGeometryCollection* GetFracturedGeometryCollection() const { return FracturedGeometryCollection; }
+
+	/** Return selection in sibling groups */
+	TMap<int32, TArray<int32>> GetClusteredSelections();
+
+	void TransformBoundsToWorld();
+
+	FBox GetBounds() const { return Bounds; }
+	void SetBounds(FBox& InBounds) { Bounds = InBounds; }
+	int32 GetSeed() const { return RandomSeed; }
+	void SetSeed(int32 InSeed) { RandomSeed = InSeed; }
+	FTransform GetTransform() const { return Transform; }
+
+	bool IsValid() const { return GeometryCollection.IsValid() && (SelectedBones.Num()>0); }
+
+private:
+	bool HasSelectedAncestor(int32 Index) const;
+	void ConvertSelectionToLeafNodes(int32 Index, TArray<int32>& LeafSelection);
+
+private:
+	TArray<int32> SelectedBones;
+	FGeometryCollectionPtr GeometryCollection;
+	UGeometryCollectionComponent* GeometryCollectionComponent;
+	UGeometryCollection* FracturedGeometryCollection;
+
+	// Fracture specific members
+	FTransform Transform;
+	FBox Bounds;
+	int32 RandomSeed;
+};
\ No newline at end of file
diff --git a/Engine/Plugins/Experimental/ChaosEditor/Source/FractureEditor/Private/FractureToolCutter.cpp b/Engine/Plugins/Experimental/ChaosEditor/Source/FractureEditor/Private/FractureToolCutter.cpp
--- a/Engine/Plugins/Experimental/ChaosEditor/Source/FractureEditor/Private/FractureToolCutter.cpp
+++ b/Engine/Plugins/Experimental/ChaosEditor/Source/FractureEditor/Private/FractureToolCutter.cpp
@@ -2,6 +2,8 @@
 
 #include "FractureToolCutter.h"
 
+#include "FractureToolContext.h"
+
 #include "GeometryCollection/GeometryCollectionObject.h"
 #include "GeometryCollection/GeometryCollectionComponent.h"
 #include "GeometryCollection/GeometryCollectionAlgo.h"
@@ -9,6 +11,8 @@
 #include "Voronoi/Voronoi.h"
 #include "PlanarCut.h"
 
+#define LOCTEXT_NAMESPACE "FractureToolCutter"
+
 
 UFractureToolCutterBase::UFractureToolCutterBase(const FObjectInitializer& ObjInit)
 	: Super(ObjInit)
@@ -29,135 +33,76 @@ bool UFractureToolCutterBase::CanExecute() const
 		return false;
 	}
 
-	return IsLeafBoneSelected();
+	return true;
 }
 
-
-void UFractureToolCutterBase::GetFractureContexts(TArray<FFractureToolContext>& FractureContexts) const
+TArray<FFractureToolContext> UFractureToolCutterBase::GetFractureToolContexts() const
 {
+	TArray<FFractureToolContext> Contexts;
+
 	FRandomStream RandomStream(CutterSettings->RandomSeed > -1 ? CutterSettings->RandomSeed : FMath::Rand());
 
-	TSet<UGeometryCollectionComponent*> GeometryCollectionComponents;
-	GetSelectedGeometryCollectionComponents(GeometryCollectionComponents);
+	// A context is gathered for each selected GeometryCollection component, or for each individual bone if Group Fracture is not used.
+	TSet<UGeometryCollectionComponent*> GeomCompSelection;
+	GetSelectedGeometryCollectionComponents(GeomCompSelection);
 
-	for (UGeometryCollectionComponent* GeometryCollectionComponent : GeometryCollectionComponents)
+	for (UGeometryCollectionComponent* GeometryCollectionComponent : GeomCompSelection)
 	{
-		FGeometryCollectionEdit RestCollection = GeometryCollectionComponent->EditRestCollection(GeometryCollection::EEditUpdate::None);
-		UGeometryCollection* FracturedGeometryCollection = RestCollection.GetRestCollection();
-		if (FracturedGeometryCollection == nullptr)
-		{
-			continue;
-		}
-
-		TSharedPtr<FGeometryCollection, ESPMode::ThreadSafe> GeometryCollectionPtr = FracturedGeometryCollection->GetGeometryCollection();
-		FGeometryCollection* OutGeometryCollection = GeometryCollectionPtr.Get();
-
-		TArray<int32> FilteredBones = FilterBones(GeometryCollectionComponent->GetSelectedBones(), OutGeometryCollection);
-
-
-		const TManagedArray<FTransform>& Transform = OutGeometryCollection->GetAttribute<FTransform>("Transform", FGeometryCollection::TransformGroup);
-		const TManagedArray<int32>& TransformToGeometryIndex = OutGeometryCollection->GetAttribute<int32>("TransformToGeometryIndex", FGeometryCollection::TransformGroup);
-		const TManagedArray<FBox>& BoundingBoxes = OutGeometryCollection->GetAttribute<FBox>("BoundingBox", FGeometryCollection::GeometryGroup);
+		// Generate a context for each selected node
+		FFractureToolContext FullSelection(GeometryCollectionComponent);
+		FullSelection.ConvertSelectionToLeafNodes();
+		
+		// Update global transforms and bounds		
+		const TManagedArray<FTransform>& Transform = FullSelection.GetGeometryCollection()->GetAttribute<FTransform>("Transform", FGeometryCollection::TransformGroup);
+		const TManagedArray<int32>& TransformToGeometryIndex = FullSelection.GetGeometryCollection()->GetAttribute<int32>("TransformToGeometryIndex", FGeometryCollection::TransformGroup);
+		const TManagedArray<FBox>& BoundingBoxes = FullSelection.GetGeometryCollection()->GetAttribute<FBox>("BoundingBox", FGeometryCollection::GeometryGroup);
 
 		TArray<FTransform> Transforms;
-		GeometryCollectionAlgo::GlobalMatrices(Transform, OutGeometryCollection->Parent, Transforms);
-
+		GeometryCollectionAlgo::GlobalMatrices(Transform, FullSelection.GetGeometryCollection()->Parent, Transforms);
 
 		TMap<int32, FBox> BoundsToBone;
-		for (int32 Idx = 0, ni = FracturedGeometryCollection->NumElements(FGeometryCollection::TransformGroup); Idx < ni; ++Idx)
+		int32 TransformCount = Transform.Num();
+		for (int32 Index = 0; Index < TransformCount; ++Index)
 		{
-			if (TransformToGeometryIndex[Idx] > -1)
+			if (TransformToGeometryIndex[Index] > INDEX_NONE)
 			{
-				ensure(TransformToGeometryIndex[Idx] > -1);
-				BoundsToBone.Add(Idx, BoundingBoxes[TransformToGeometryIndex[Idx]].TransformBy(Transforms[Idx]));
+				BoundsToBone.Add(Index, BoundingBoxes[TransformToGeometryIndex[Index]].TransformBy(Transforms[Index]));
 			}
 		}
 
 		if (CutterSettings->bGroupFracture)
 		{
-			GenerateFractureToolContext(
-				GeometryCollectionComponent->GetOwner(),
-				GeometryCollectionComponent,
-				FracturedGeometryCollection,
-				FilteredBones,
-				BoundsToBone,
-				TransformToGeometryIndex,
-				CutterSettings->RandomSeed,
-				FractureContexts);
-		}
-		else
-		{
-			for (int32 BoneIndex : FilteredBones)
-			{
-				GenerateFractureToolContext(
-					GeometryCollectionComponent->GetOwner(),
-					GeometryCollectionComponent,
-					FracturedGeometryCollection,
-					{ BoneIndex },
-					BoundsToBone,
-					TransformToGeometryIndex,
-					CutterSettings->RandomSeed + FractureContexts.Num(),
-					FractureContexts);
-			}
-		}
-		
-	}
-}
+			FullSelection.SetSeed(CutterSettings->RandomSeed == -1 ? CutterSettings->RandomSeed : FMath::Rand());
 
-
-bool UFractureToolCutterBase::IsLeafBoneSelected()
-{
-	TSet<UGeometryCollectionComponent*> GeomCompSelection;
-	GetSelectedGeometryCollectionComponents(GeomCompSelection);
-	for (UGeometryCollectionComponent* GeometryCollectionComponent : GeomCompSelection)
-	{
-		const TArray<int32> SelectedBones = GeometryCollectionComponent->GetSelectedBones();
-
-		if (SelectedBones.Num() > 0)
-		{
-			if (const UGeometryCollection* GCObject = GeometryCollectionComponent->GetRestCollection())
+			FBox Bounds(ForceInit);
+			for (int32 BoneIndex : FullSelection.GetSelection())
 			{
-				TSharedPtr<FGeometryCollection, ESPMode::ThreadSafe> GeometryCollectionPtr = GCObject->GetGeometryCollection();
-				if (const FGeometryCollection* GeometryCollection = GeometryCollectionPtr.Get())
+				if (TransformToGeometryIndex[BoneIndex] > INDEX_NONE)
 				{
-					const TManagedArray<TSet<int32>>& Children = GeometryCollection->GetAttribute<TSet<int32>>("Children", FGeometryCollection::TransformGroup);
-
-					for (int32 BoneIndex : SelectedBones)
-					{
-						if (Children[BoneIndex].Num() == 0)
-						{
-							return true;
-						}
-					}
+					Bounds += BoundsToBone[BoneIndex];
 				}
 			}
-		}
-	}
-	return false;
-}
-
-TArray<int32> UFractureToolCutterBase::FilterBones(const TArray<int32>& SelectedBonesOriginal, const FGeometryCollection* const GeometryCollection) const
-{
-	FRandomStream RandomStream(CutterSettings->RandomSeed > -1 ? CutterSettings->RandomSeed : FMath::Rand());
+			FullSelection.SetBounds(Bounds);
 
-	// Keep only leaf nodes and implement skip probability
-	const TManagedArray<TSet<int32>>& Children = GeometryCollection->GetAttribute<TSet<int32>>("Children", FGeometryCollection::TransformGroup);
-
-	TArray<int32> SelectedBones;
-	SelectedBones.Reserve(SelectedBonesOriginal.Num());
-	for (int32 BoneIndex : SelectedBonesOriginal)
-	{
-		if (Children[BoneIndex].Num() == 0 && RandomStream.FRand() < CutterSettings->ChanceToFracture)
+			Contexts.Add(FullSelection);
+		}
+		else
 		{
-			SelectedBones.Add(BoneIndex);
+			// Generate a context for each selected node
+			for (int32 Index : FullSelection.GetSelection())
+			{
+				Contexts.Emplace(GeometryCollectionComponent);
+				FFractureToolContext& FractureContext = Contexts.Last();
+
+				FractureContext.SetSeed(CutterSettings->RandomSeed == -1 ? CutterSettings->RandomSeed + Index : FMath::Rand());
+				FractureContext.SetBounds(BoundsToBone[Index]);
+			}
 		}
 	}
 
-	return SelectedBones;
+	return Contexts;
 }
 
-
-
 UFractureToolVoronoiCutterBase::UFractureToolVoronoiCutterBase(const FObjectInitializer& ObjInit)
 	: Super(ObjInit)
 {
@@ -189,8 +134,7 @@ void UFractureToolVoronoiCutterBase::Render(const FSceneView* View, FViewport* V
 
 void UFractureToolVoronoiCutterBase::FractureContextChanged()
 {
-	TArray<FFractureToolContext> FractureContexts;
-	GetFractureContexts(FractureContexts);
+	TArray<FFractureToolContext> FractureContexts = GetFractureToolContexts();
 
 	VoronoiSites.Empty();
 	CellMember.Empty();
@@ -199,45 +143,45 @@ void UFractureToolVoronoiCutterBase::FractureContextChanged()
 	for (FFractureToolContext& FractureContext : FractureContexts)
 	{
 		// Move the local bounds to the actor so we we'll draw in the correct location
-		FractureContext.Bounds = FractureContext.Bounds.TransformBy(FractureContext.Transform);
+		FractureContext.TransformBoundsToWorld();
 		GenerateVoronoiSites(FractureContext, VoronoiSites);
 		if (CutterSettings->bDrawDiagram)
 		{
-			GetVoronoiEdges(VoronoiSites, FractureContext.Bounds, VoronoiEdges, CellMember);
+			GetVoronoiEdges(VoronoiSites, FractureContext.GetBounds(), VoronoiEdges, CellMember);
 		}
 	}
 }
 
-void UFractureToolVoronoiCutterBase::ExecuteFracture(const FFractureToolContext& FractureContext)
+int32 UFractureToolVoronoiCutterBase::ExecuteFracture(const FFractureToolContext& FractureContext)
 {
-	if (FractureContext.FracturedGeometryCollection != nullptr)
+	if (FractureContext.IsValid())
 	{
-		TSharedPtr<FGeometryCollection, ESPMode::ThreadSafe> GeometryCollectionPtr = FractureContext.FracturedGeometryCollection->GetGeometryCollection();
-		if (FGeometryCollection* GeometryCollection = GeometryCollectionPtr.Get())
+		TArray<FVector> Sites;
+		GenerateVoronoiSites(FractureContext, Sites);
+		FBox VoronoiBounds = FractureContext.GetBounds();
+		if (CutterSettings->Amplitude > 0.0f)
 		{
-			TArray<FVector> Sites;
-			GenerateVoronoiSites(FractureContext, Sites);
-			FBox VoronoiBounds = FractureContext.Bounds;
-			if (CutterSettings->Amplitude > 0.0f)
-			{
-				// expand bounds to make sure noise-perturbed voronoi cells still contain the whole input mesh
-				VoronoiBounds = VoronoiBounds.ExpandBy(CutterSettings->Amplitude);
-			}
-			FVoronoiDiagram Voronoi(Sites, VoronoiBounds, .1f);
-
-			FPlanarCells VoronoiPlanarCells = FPlanarCells(Sites, Voronoi);
+			// expand bounds to make sure noise-perturbed voronoi cells still contain the whole input mesh
+			VoronoiBounds = VoronoiBounds.ExpandBy(CutterSettings->Amplitude);
+		}
+		FVoronoiDiagram Voronoi(Sites, VoronoiBounds, .1f);
 
-			FNoiseSettings NoiseSettings;
-			if (CutterSettings->Amplitude > 0.0f)
-			{
-				NoiseSettings.Amplitude = CutterSettings->Amplitude;
-				NoiseSettings.Frequency = CutterSettings->Frequency;
-				NoiseSettings.Octaves = CutterSettings->OctaveNumber;
-				NoiseSettings.PointSpacing = CutterSettings->SurfaceResolution;
-				VoronoiPlanarCells.InternalSurfaceMaterials.NoiseSettings = NoiseSettings;
-			}
+		FPlanarCells VoronoiPlanarCells = FPlanarCells(Sites, Voronoi);
 
-			CutMultipleWithPlanarCells(VoronoiPlanarCells, *GeometryCollection, FractureContext.SelectedBones);
+		FNoiseSettings NoiseSettings;
+		if (CutterSettings->Amplitude > 0.0f)
+		{
+			NoiseSettings.Amplitude = CutterSettings->Amplitude;
+			NoiseSettings.Frequency = CutterSettings->Frequency;
+			NoiseSettings.Octaves = CutterSettings->OctaveNumber;
+			NoiseSettings.PointSpacing = CutterSettings->SurfaceResolution;
+			VoronoiPlanarCells.InternalSurfaceMaterials.NoiseSettings = NoiseSettings;
 		}
+
+		return CutMultipleWithPlanarCells(VoronoiPlanarCells, *(FractureContext.GetGeometryCollection()), FractureContext.GetSelection());
 	}
+
+	return INDEX_NONE;
 }
+
+#undef LOCTEXT_NAMESPACE
diff --git a/Engine/Plugins/Experimental/ChaosEditor/Source/FractureEditor/Private/FractureToolCutter.h b/Engine/Plugins/Experimental/ChaosEditor/Source/FractureEditor/Private/FractureToolCutter.h
--- a/Engine/Plugins/Experimental/ChaosEditor/Source/FractureEditor/Private/FractureToolCutter.h
+++ b/Engine/Plugins/Experimental/ChaosEditor/Source/FractureEditor/Private/FractureToolCutter.h
@@ -6,6 +6,7 @@
 
 #include "FractureToolCutter.generated.h"
 
+class FFractureToolContext;
 
 /** Settings specifically related to the one-time destructive fracturing of a mesh **/
 UCLASS(config = EditorPerProjectUserSettings)
@@ -80,17 +81,12 @@ public:
 
 	virtual bool CanExecute() const override;
 
-	virtual void GetFractureContexts(TArray<FFractureToolContext>& FractureContexts) const;
-
-	TArray<int32> FilterBones(const TArray<int32>& SelectedBonesOriginal, const FGeometryCollection* const GeometryCollection) const override;
+	virtual TArray<FFractureToolContext> GetFractureToolContexts() const override;
 
 protected:
 	UPROPERTY(EditAnywhere, Category = Slicing)
 	UFractureCutterSettings* CutterSettings;
 
-private:
-	static bool IsLeafBoneSelected();
-
 };
 
 
@@ -105,7 +101,7 @@ public:
 	virtual void Render(const FSceneView* View, FViewport* Viewport, FPrimitiveDrawInterface* PDI) override;
 
 	virtual void FractureContextChanged() override;
-	virtual void ExecuteFracture(const FFractureToolContext& FractureContext) override;
+	virtual int32 ExecuteFracture(const FFractureToolContext& FractureContext) override;
 
 protected:
 	virtual void GenerateVoronoiSites(const FFractureToolContext& Context, TArray<FVector>& Sites) {}
diff --git a/Engine/Plugins/Experimental/ChaosEditor/Source/FractureEditor/Private/FractureToolPlaneCut.cpp b/Engine/Plugins/Experimental/ChaosEditor/Source/FractureEditor/Private/FractureToolPlaneCut.cpp
--- a/Engine/Plugins/Experimental/ChaosEditor/Source/FractureEditor/Private/FractureToolPlaneCut.cpp
+++ b/Engine/Plugins/Experimental/ChaosEditor/Source/FractureEditor/Private/FractureToolPlaneCut.cpp
@@ -6,6 +6,7 @@
 #include "GeometryCollection/GeometryCollection.h"
 #include "GeometryCollection/GeometryCollectionObject.h"
 #include "FractureEditorModeToolkit.h"
+#include "FractureToolContext.h"
 
 #define LOCTEXT_NAMESPACE "FracturePlanar"
 
@@ -78,76 +79,73 @@ TArray<UObject*> UFractureToolPlaneCut::GetSettingsObjects() const
 
 void UFractureToolPlaneCut::FractureContextChanged()
 {
-	TArray<FFractureToolContext> FractureContexts;
-	GetFractureContexts(FractureContexts);
+	TArray<FFractureToolContext> FractureContexts = GetFractureToolContexts();
 
 	RenderCuttingPlanesTransforms.Empty();
 
 	RenderCuttingPlaneSize = FLT_MAX;
 	for (FFractureToolContext& FractureContext : FractureContexts)
 	{
-		// Move the local bounds to the actor so we weill draw in the correct location
-		FractureContext.Bounds = FractureContext.Bounds.TransformBy(FractureContext.Transform);
+		// Move the local bounds to the actor so we we'll draw in the correct location
+		FractureContext.TransformBoundsToWorld();
 		GenerateSliceTransforms(FractureContext, RenderCuttingPlanesTransforms);
 
-		if (FractureContext.Bounds.GetExtent().GetMax() < RenderCuttingPlaneSize)
+		if (FractureContext.GetBounds().GetExtent().GetMax() < RenderCuttingPlaneSize)
 		{
-			RenderCuttingPlaneSize = FractureContext.Bounds.GetExtent().GetMax();
+			RenderCuttingPlaneSize = FractureContext.GetBounds().GetExtent().GetMax();
 		}
 	}
 }
 
-void UFractureToolPlaneCut::ExecuteFracture(const FFractureToolContext& FractureContext)
+int32 UFractureToolPlaneCut::ExecuteFracture(const FFractureToolContext& FractureContext)
 {
-	if (FractureContext.FracturedGeometryCollection != nullptr)
+	if (FractureContext.IsValid())
 	{
-		TSharedPtr<FGeometryCollection, ESPMode::ThreadSafe> GeometryCollectionPtr = FractureContext.FracturedGeometryCollection->GetGeometryCollection();
-		if (FGeometryCollection* GeometryCollection = GeometryCollectionPtr.Get())
-		{
-			const UFracturePlaneCutSettings* LocalCutSettings = GetMutableDefault<UFracturePlaneCutSettings>();
+		const UFracturePlaneCutSettings* LocalCutSettings = GetMutableDefault<UFracturePlaneCutSettings>();
 
-			TArray<FPlane> CuttingPlanes;
-			if (LocalCutSettings->ReferenceActor != nullptr)
-			{
-				FTransform Transform(LocalCutSettings->ReferenceActor->GetActorTransform());
-				CuttingPlanes.Add(FPlane(Transform.GetLocation() - FractureContext.Transform.GetLocation(), Transform.GetUnitAxis(EAxis::Z)));
-			}
-			else
-			{
-				TArray<FTransform> CuttingPlaneTransforms;
-				GenerateSliceTransforms(FractureContext, CuttingPlaneTransforms);
-				for (const FTransform& Transform : CuttingPlaneTransforms)
-				{
-					CuttingPlanes.Add(FPlane(Transform.GetLocation(), Transform.GetUnitAxis(EAxis::Z)));
-				}
-			}
-
-			FInternalSurfaceMaterials InternalSurfaceMaterials;
-			FNoiseSettings NoiseSettings;
-			if (CutterSettings->Amplitude > 0.0f)
+		TArray<FPlane> CuttingPlanes;
+		if (LocalCutSettings->ReferenceActor != nullptr)
+		{
+			FTransform Transform(LocalCutSettings->ReferenceActor->GetActorTransform());
+			CuttingPlanes.Add(FPlane(Transform.GetLocation() - FractureContext.GetTransform().GetLocation(), Transform.GetUnitAxis(EAxis::Z)));
+		}
+		else
+		{
+			TArray<FTransform> CuttingPlaneTransforms;
+			GenerateSliceTransforms(FractureContext, CuttingPlaneTransforms);
+			for (const FTransform& Transform : CuttingPlaneTransforms)
 			{
-				NoiseSettings.Amplitude = CutterSettings->Amplitude;
-				NoiseSettings.Frequency = CutterSettings->Frequency;
-				NoiseSettings.Octaves = CutterSettings->OctaveNumber;
-				NoiseSettings.PointSpacing = CutterSettings->SurfaceResolution;
-				InternalSurfaceMaterials.NoiseSettings = NoiseSettings;
+				CuttingPlanes.Add(FPlane(Transform.GetLocation(), Transform.GetUnitAxis(EAxis::Z)));
 			}
+		}
 
-			CutMultipleWithMultiplePlanes(CuttingPlanes, InternalSurfaceMaterials, *GeometryCollection, FractureContext.SelectedBones);
+		FInternalSurfaceMaterials InternalSurfaceMaterials;
+		FNoiseSettings NoiseSettings;
+		if (CutterSettings->Amplitude > 0.0f)
+		{
+			NoiseSettings.Amplitude = CutterSettings->Amplitude;
+			NoiseSettings.Frequency = CutterSettings->Frequency;
+			NoiseSettings.Octaves = CutterSettings->OctaveNumber;
+			NoiseSettings.PointSpacing = CutterSettings->SurfaceResolution;
+			InternalSurfaceMaterials.NoiseSettings = NoiseSettings;
 		}
+
+		return CutMultipleWithMultiplePlanes(CuttingPlanes, InternalSurfaceMaterials, *FractureContext.GetGeometryCollection(), FractureContext.GetSelection());
 	}
+
+	return INDEX_NONE;
 }
 
 void UFractureToolPlaneCut::GenerateSliceTransforms(const FFractureToolContext& Context, TArray<FTransform>& CuttingPlaneTransforms)
 {
-	FRandomStream RandStream(Context.RandomSeed);
+	FRandomStream RandStream(Context.GetSeed());
 
-	const FVector Extent(Context.Bounds.Max - Context.Bounds.Min);
+	const FVector Extent(Context.GetBounds().Max - Context.GetBounds().Min);
 
 	CuttingPlaneTransforms.Reserve(CuttingPlaneTransforms.Num() + PlaneCutSettings->NumberPlanarCuts);
 	for (int32 ii = 0; ii < PlaneCutSettings->NumberPlanarCuts; ++ii)
 	{
-		FVector Position(Context.Bounds.Min + FVector(RandStream.FRand(), RandStream.FRand(), RandStream.FRand()) * Extent);
+		FVector Position(Context.GetBounds().Min + FVector(RandStream.FRand(), RandStream.FRand(), RandStream.FRand()) * Extent);
 		CuttingPlaneTransforms.Emplace(FTransform(FRotator(RandStream.FRand() * 360.0f, RandStream.FRand() * 360.0f, 0.0f), Position));
 	}
 }
diff --git a/Engine/Plugins/Experimental/ChaosEditor/Source/FractureEditor/Private/FractureToolPlaneCut.h b/Engine/Plugins/Experimental/ChaosEditor/Source/FractureEditor/Private/FractureToolPlaneCut.h
--- a/Engine/Plugins/Experimental/ChaosEditor/Source/FractureEditor/Private/FractureToolPlaneCut.h
+++ b/Engine/Plugins/Experimental/ChaosEditor/Source/FractureEditor/Private/FractureToolPlaneCut.h
@@ -6,6 +6,8 @@
 
 #include "FractureToolPlaneCut.generated.h"
 
+class FFractureToolContext;
+
 
 UCLASS(config = EditorPerProjectUserSettings)
 class UFracturePlaneCutSettings : public UFractureToolSettings
@@ -48,7 +50,7 @@ public:
 	virtual TArray<UObject*> GetSettingsObjects() const override;
 
 	virtual void FractureContextChanged() override;
-	virtual void ExecuteFracture(const FFractureToolContext& FractureContext) override;
+	virtual int32 ExecuteFracture(const FFractureToolContext& FractureContext) override;
 
 private:
 	// Slicing
diff --git a/Engine/Plugins/Experimental/ChaosEditor/Source/FractureEditor/Private/FractureToolRadial.cpp b/Engine/Plugins/Experimental/ChaosEditor/Source/FractureEditor/Private/FractureToolRadial.cpp
--- a/Engine/Plugins/Experimental/ChaosEditor/Source/FractureEditor/Private/FractureToolRadial.cpp
+++ b/Engine/Plugins/Experimental/ChaosEditor/Source/FractureEditor/Private/FractureToolRadial.cpp
@@ -2,7 +2,8 @@
 
 #include "FractureToolRadial.h"
 
-#include "FractureEditorStyle.h"
+#include "FractureToolContext.h"
+
 
 #define LOCTEXT_NAMESPACE "FractureRadial"
 
@@ -43,13 +44,13 @@ TArray<UObject*> UFractureToolRadial::GetSettingsObjects() const
 	return Settings;
 }
 
-void UFractureToolRadial::GenerateVoronoiSites(const FFractureToolContext&Context, TArray<FVector>& Sites)
+void UFractureToolRadial::GenerateVoronoiSites(const FFractureToolContext& Context, TArray<FVector>& Sites)
 {
  	float RadialStep = RadialSettings->Radius / RadialSettings->RadialSteps;
 
-	const FVector Center(Context.Bounds.GetCenter() + RadialSettings->Center);
+	const FVector Center(Context.GetBounds().GetCenter() + RadialSettings->Center);
 
-	FRandomStream RandStream(Context.RandomSeed);
+	FRandomStream RandStream(Context.GetSeed());
 	FVector UpVector(RadialSettings->Normal);
 	UpVector.Normalize();
 	FVector PerpVector(UpVector[2], UpVector[0], UpVector[1]);
diff --git a/Engine/Plugins/Experimental/ChaosEditor/Source/FractureEditor/Private/FractureToolRadial.h b/Engine/Plugins/Experimental/ChaosEditor/Source/FractureEditor/Private/FractureToolRadial.h
--- a/Engine/Plugins/Experimental/ChaosEditor/Source/FractureEditor/Private/FractureToolRadial.h
+++ b/Engine/Plugins/Experimental/ChaosEditor/Source/FractureEditor/Private/FractureToolRadial.h
@@ -7,6 +7,7 @@
 
 #include "FractureToolRadial.generated.h"
 
+class FFractureToolContext;
 
 UCLASS(config = EditorPerProjectUserSettings)
 class UFractureRadialSettings : public UFractureToolSettings
@@ -76,6 +77,6 @@ public:
 	UFractureRadialSettings* RadialSettings;
 
 protected:
-	void GenerateVoronoiSites(const FFractureToolContext &Context, TArray<FVector>& Sites) override;
+	void GenerateVoronoiSites(const FFractureToolContext& Context, TArray<FVector>& Sites) override;
 
 };
\ No newline at end of file
diff --git a/Engine/Plugins/Experimental/ChaosEditor/Source/FractureEditor/Private/FractureToolSlice.cpp b/Engine/Plugins/Experimental/ChaosEditor/Source/FractureEditor/Private/FractureToolSlice.cpp
--- a/Engine/Plugins/Experimental/ChaosEditor/Source/FractureEditor/Private/FractureToolSlice.cpp
+++ b/Engine/Plugins/Experimental/ChaosEditor/Source/FractureEditor/Private/FractureToolSlice.cpp
@@ -10,6 +10,7 @@
 #include "GeometryCollection/GeometryCollectionObject.h"
 #include "PlanarCut.h"
 #include "FractureEditorModeToolkit.h"
+#include "FractureToolContext.h"
 
 #define LOCTEXT_NAMESPACE "FractureSlice"
 
@@ -53,18 +54,18 @@ TArray<UObject*> UFractureToolSlice::GetSettingsObjects() const
 
 void UFractureToolSlice::GenerateSliceTransforms(const FFractureToolContext& Context, TArray<FTransform>& CuttingPlaneTransforms)
 {
-	const FBox& Bounds = Context.Bounds;
-	const FVector& Min = Context.Bounds.Min;
-	const FVector& Max = Context.Bounds.Max;
-	const FVector  Center = Context.Bounds.GetCenter();
-	const FVector Extents(Context.Bounds.Max - Context.Bounds.Min);
+	const FBox& Bounds = Context.GetBounds();
+	const FVector& Min = Context.GetBounds().Min;
+	const FVector& Max = Context.GetBounds().Max;
+	const FVector  Center = Context.GetBounds().GetCenter();
+	const FVector Extents(Context.GetBounds().Max - Context.GetBounds().Min);
 	const FVector HalfExtents(Extents * 0.5f);
 
 	const FVector Step(Extents.X / (SliceSettings->SlicesX+1), Extents.Y / (SliceSettings->SlicesY+1), Extents.Z / (SliceSettings->SlicesZ+1));
 
 	CuttingPlaneTransforms.Reserve(SliceSettings->SlicesX * SliceSettings->SlicesY * SliceSettings->SlicesZ);
 
-	FRandomStream RandomStream(Context.RandomSeed);
+	FRandomStream RandomStream(Context.GetSeed());
 
 	const float SliceAngleVariationInRadians = FMath::DegreesToRadians(SliceSettings->SliceAngleVariation);
 
@@ -124,57 +125,54 @@ void UFractureToolSlice::Render(const FSceneView* View, FViewport* Viewport, FPr
 
 void UFractureToolSlice::FractureContextChanged()
 {
-	TArray<FFractureToolContext> FractureContexts;
-	GetFractureContexts(FractureContexts);
+	TArray<FFractureToolContext> FractureContexts = GetFractureToolContexts();
 
 	RenderCuttingPlanesTransforms.Empty();
 
 	RenderCuttingPlaneSize = FLT_MAX;
 	for (FFractureToolContext& FractureContext : FractureContexts)
 	{
-		// Move the local bounds to the actor so we weill draw in the correct location
-		FractureContext.Bounds = FractureContext.Bounds.TransformBy(FractureContext.Transform);
+		// Move the local bounds to the actor so we we'll draw in the correct location
+		FractureContext.TransformBoundsToWorld();
 		GenerateSliceTransforms(FractureContext, RenderCuttingPlanesTransforms);
 
-		if (FractureContext.Bounds.GetExtent().GetMax() < RenderCuttingPlaneSize)
+		if (FractureContext.GetBounds().GetExtent().GetMax() < RenderCuttingPlaneSize)
 		{
-			RenderCuttingPlaneSize = FractureContext.Bounds.GetExtent().GetMax();
+			RenderCuttingPlaneSize = FractureContext.GetBounds().GetExtent().GetMax();
 		}
 	}
 }
 
-void UFractureToolSlice::ExecuteFracture(const FFractureToolContext& FractureContext)
+int32 UFractureToolSlice::ExecuteFracture(const FFractureToolContext& FractureContext)
 {
-	if (FractureContext.FracturedGeometryCollection != nullptr)
+	if(FractureContext.IsValid())
 	{
-		TSharedPtr<FGeometryCollection, ESPMode::ThreadSafe> GeometryCollectionPtr = FractureContext.FracturedGeometryCollection->GetGeometryCollection();
-		if (FGeometryCollection* GeometryCollection = GeometryCollectionPtr.Get())
+		TArray<FTransform> LocalCuttingPlanesTransforms;
+		GenerateSliceTransforms(FractureContext, LocalCuttingPlanesTransforms);
+
+		TArray<FPlane> CuttingPlanes;
+		CuttingPlanes.Reserve(LocalCuttingPlanesTransforms.Num());
+
+		for (const FTransform& Transform : LocalCuttingPlanesTransforms)
 		{
-			TArray<FTransform> LocalCuttingPlanesTransforms;
-			GenerateSliceTransforms(FractureContext, LocalCuttingPlanesTransforms);
-
-			TArray<FPlane> CuttingPlanes;
-			CuttingPlanes.Reserve(LocalCuttingPlanesTransforms.Num());
-
-			for (const FTransform& Transform : LocalCuttingPlanesTransforms)
-			{
-				CuttingPlanes.Add(FPlane(Transform.GetLocation(), Transform.GetUnitAxis(EAxis::Z)));
-			}
-
-			FInternalSurfaceMaterials InternalSurfaceMaterials;
-			FNoiseSettings NoiseSettings;
-			if (CutterSettings->Amplitude > 0.0f)
-			{
-				NoiseSettings.Amplitude = CutterSettings->Amplitude;
-				NoiseSettings.Frequency = CutterSettings->Frequency;
-				NoiseSettings.Octaves = CutterSettings->OctaveNumber;
-				NoiseSettings.PointSpacing = CutterSettings->SurfaceResolution;
-				InternalSurfaceMaterials.NoiseSettings = NoiseSettings;
-			}
-
-			int32 AddedGeomIdx = CutMultipleWithMultiplePlanes(CuttingPlanes, InternalSurfaceMaterials, *GeometryCollection, FractureContext.SelectedBones);
+			CuttingPlanes.Add(FPlane(Transform.GetLocation(), Transform.GetUnitAxis(EAxis::Z)));
 		}
+
+		FInternalSurfaceMaterials InternalSurfaceMaterials;
+		FNoiseSettings NoiseSettings;
+		if (CutterSettings->Amplitude > 0.0f)
+		{
+			NoiseSettings.Amplitude = CutterSettings->Amplitude;
+			NoiseSettings.Frequency = CutterSettings->Frequency;
+			NoiseSettings.Octaves = CutterSettings->OctaveNumber;
+			NoiseSettings.PointSpacing = CutterSettings->SurfaceResolution;
+			InternalSurfaceMaterials.NoiseSettings = NoiseSettings;
+		}
+
+		return CutMultipleWithMultiplePlanes(CuttingPlanes, InternalSurfaceMaterials, *FractureContext.GetGeometryCollection(), FractureContext.GetSelection());
 	}
+
+	return INDEX_NONE;
 }
 
 #undef LOCTEXT_NAMESPACE
\ No newline at end of file
diff --git a/Engine/Plugins/Experimental/ChaosEditor/Source/FractureEditor/Private/FractureToolSlice.h b/Engine/Plugins/Experimental/ChaosEditor/Source/FractureEditor/Private/FractureToolSlice.h
--- a/Engine/Plugins/Experimental/ChaosEditor/Source/FractureEditor/Private/FractureToolSlice.h
+++ b/Engine/Plugins/Experimental/ChaosEditor/Source/FractureEditor/Private/FractureToolSlice.h
@@ -7,6 +7,9 @@
 #include "FractureToolSlice.generated.h"
 
 
+class FFractureToolContext;
+
+
 UCLASS(config = EditorPerProjectUserSettings)
 class UFractureSliceSettings : public UFractureToolSettings
 {
@@ -62,7 +65,7 @@ public:
 	void Render(const FSceneView* View, FViewport* Viewport, FPrimitiveDrawInterface* PDI) override;
 
 	virtual void FractureContextChanged() override;
-	virtual void ExecuteFracture(const FFractureToolContext& Context) override;
+	virtual int32 ExecuteFracture(const FFractureToolContext& Context) override;
 
 	void GenerateSliceTransforms(const FFractureToolContext& Context, TArray<FTransform>& CuttingPlaneTransforms);
 
diff --git a/Engine/Plugins/Experimental/ChaosEditor/Source/FractureEditor/Private/FractureToolUniform.cpp b/Engine/Plugins/Experimental/ChaosEditor/Source/FractureEditor/Private/FractureToolUniform.cpp
--- a/Engine/Plugins/Experimental/ChaosEditor/Source/FractureEditor/Private/FractureToolUniform.cpp
+++ b/Engine/Plugins/Experimental/ChaosEditor/Source/FractureEditor/Private/FractureToolUniform.cpp
@@ -4,6 +4,7 @@
 
 #include "FractureEditorStyle.h"
 #include "FractureEditorCommands.h"
+#include "FractureToolContext.h"
 
 #define LOCTEXT_NAMESPACE "FractureUniform"
 
@@ -44,18 +45,18 @@ TArray<UObject*> UFractureToolUniform::GetSettingsObjects() const
 	return AllSettings;
 }
 
-void UFractureToolUniform::GenerateVoronoiSites(const FFractureToolContext &Context, TArray<FVector>& Sites)
+void UFractureToolUniform::GenerateVoronoiSites(const FFractureToolContext& Context, TArray<FVector>& Sites)
 {
-	FRandomStream RandStream(Context.RandomSeed);
+	FRandomStream RandStream(Context.GetSeed());
 
-	const FVector Extent(Context.Bounds.Max - Context.Bounds.Min);
+	const FVector Extent(Context.GetBounds().Max - Context.GetBounds().Min);
 
 	const int32 SiteCount = RandStream.RandRange(UniformSettings->NumberVoronoiSitesMin, UniformSettings->NumberVoronoiSitesMax);
 
 	Sites.Reserve(Sites.Num() + SiteCount);
 	for (int32 ii = 0; ii < SiteCount; ++ii)
 	{
-		Sites.Emplace(Context.Bounds.Min + FVector(RandStream.FRand(), RandStream.FRand(), RandStream.FRand()) * Extent );
+		Sites.Emplace(Context.GetBounds().Min + FVector(RandStream.FRand(), RandStream.FRand(), RandStream.FRand()) * Extent );
 	}
 }
 
diff --git a/Engine/Plugins/Experimental/ChaosEditor/Source/FractureEditor/Private/FractureToolUniform.h b/Engine/Plugins/Experimental/ChaosEditor/Source/FractureEditor/Private/FractureToolUniform.h
--- a/Engine/Plugins/Experimental/ChaosEditor/Source/FractureEditor/Private/FractureToolUniform.h
+++ b/Engine/Plugins/Experimental/ChaosEditor/Source/FractureEditor/Private/FractureToolUniform.h
@@ -6,7 +6,7 @@
 
 #include "FractureToolUniform.generated.h"
 
-
+class FFractureToolContext;
 
 UCLASS(config = EditorPerProjectUserSettings)
 class UFractureUniformSettings : public UFractureToolSettings
@@ -52,6 +52,6 @@ public:
 	UFractureUniformSettings* UniformSettings;
 
 protected:
-	void GenerateVoronoiSites(const FFractureToolContext&Context, TArray<FVector>& Sites) override;
+	void GenerateVoronoiSites(const FFractureToolContext& Context, TArray<FVector>& Sites) override;
 };
 
diff --git a/Engine/Plugins/Experimental/ChaosEditor/Source/FractureEditor/Private/SGeometryCollectionHistogram.cpp b/Engine/Plugins/Experimental/ChaosEditor/Source/FractureEditor/Private/SGeometryCollectionHistogram.cpp
new file mode 100644
--- /dev/null
+++ b/Engine/Plugins/Experimental/ChaosEditor/Source/FractureEditor/Private/SGeometryCollectionHistogram.cpp
@@ -0,0 +1,371 @@
+// Copyright Epic Games, Inc. All Rights Reserved.
+
+#include "SGeometryCollectionHistogram.h"
+
+#include "FractureToolProperties.h"
+
+#include "GeometryCollection/GeometryCollectionComponent.h"
+
+#include "Widgets/Colors/SColorBlock.h"
+#include "Widgets/Text/STextBlock.h"
+
+#include "FractureEditorMode.h"
+#include "ScopedTransaction.h"
+
+#define LOCTEXT_NAMESPACE "ChaosEditor"
+
+
+UHistogramSettings::UHistogramSettings(const FObjectInitializer& ObjInit)
+	: Super(ObjInit)
+	, InspectedAttribute(EInspectedAttributeEnum::Volume)
+	, bSorted(true)
+{}
+
+TSharedRef<ITableRow> FGeometryCollectionHistogramItem::MakeHistogramRowWidget(const TSharedRef<STableViewBase>& InOwnerTable)
+{
+	return SNew(STableRow<FGeometryCollectionHistogramItemPtr>, InOwnerTable)
+		.Padding(FMargin(1.0f))
+		.Content()
+		[
+			SNew(SHorizontalBox)
+			+SHorizontalBox::Slot()
+			.FillWidth(NormalizedValue)
+			[
+				SNew(SColorBlock)
+				.Color(NodeColor)
+				.AlphaDisplayMode(EColorBlockAlphaDisplayMode::Ignore)
+				.Size(FVector2D(1.0f,5.0f))
+				.ToolTipText(HoverText)
+			]
+			+ SHorizontalBox::Slot()
+			.FillWidth(1.0-NormalizedValue)
+			[
+				SNew(SBox)
+			]
+		];
+}
+
+void FGeometryCollectionHistogramItem::SetInspectedAttribute(EInspectedAttributeEnum InspectedAttribute)
+{
+	if (ensure(ParentComponentItem.IsValid()))
+	{
+		UGeometryCollectionComponent* Component = ParentComponentItem->GetComponent();
+
+		if (Component)
+		{
+			FGeometryCollection* Collection = Component->GetRestCollection()->GetGeometryCollection().Get();
+
+			switch (InspectedAttribute)
+			{
+			case EInspectedAttributeEnum::Volume:
+				InspectedValue = Collection->GetAttribute<float>(TEXT("Volume"), TEXT("Transform"))[BoneIndex];
+				HoverText = FText::FromString(FString::SanitizeFloat(InspectedValue));
+				break;
+
+			case EInspectedAttributeEnum::Level:
+				{
+					int32 Level = Collection->GetAttribute<int32>(TEXT("Level"), TEXT("Transform"))[BoneIndex];
+					HoverText = FText::FromString(FString::FormatAsNumber(Level));
+					InspectedValue = static_cast<float>(Level);
+				}
+				break;
+
+			case EInspectedAttributeEnum::InitialDynamicState:
+				{
+					int32 InitialDynamicState = static_cast<int32>(Collection->GetAttribute<int32>(TEXT("InitialDynamicState"), TEXT("Transform"))[BoneIndex]);
+					
+					static const TArray<FString> HoverNames{ "No Override", "Sleeping", "Kinematic", "Static" };
+					HoverText = FText::FromString(HoverNames[InitialDynamicState]);
+
+					InspectedValue = static_cast<float>(InitialDynamicState);
+				}
+				break;
+
+			default:
+				// Invalid inspection attribute
+				check(false);
+			}
+		}	
+	}
+}
+
+
+void FGeometryCollectionHistogramItem::SetNormalizedValue(float MinValue, float MaxValue)
+{
+	if ((MaxValue - MinValue) > KINDA_SMALL_NUMBER)
+	{
+		NormalizedValue = (InspectedValue - MinValue) / (MaxValue - MinValue);
+	} 
+	else
+	{
+		NormalizedValue = 1.0;
+	}
+}
+
+
+FGeometryCollectionHistogramItemPtr FGeometryCollectionHistogramItemComponent::GetItemFromBoneIndex(int32 BoneIndex) const
+{
+	for (auto& Pair : NodesMap)
+	{
+		if (Pair.Value->GetBoneIndex() == BoneIndex)
+		{
+			return Pair.Value;
+		}
+	}
+
+	return FGeometryCollectionHistogramItemPtr();
+}
+
+FGeometryCollectionHistogramItemList FGeometryCollectionHistogramItemComponent::RegenerateNodes()
+{
+	// Collect the inspected attribute 
+	FGeometryCollectionHistogramItemList NodesList;
+	
+	if (Component->GetRestCollection())
+	{
+		FGeometryCollection* Collection = Component->GetRestCollection()->GetGeometryCollection().Get();
+
+		NodesMap.Empty();
+		GuidIndexMap.Empty();
+
+		if (Collection)
+		{
+			int32 NumElements = Collection->NumElements(FGeometryCollection::TransformGroup);
+
+			const TManagedArray<FGuid>& Guids = Collection->GetAttribute<FGuid>("GUID", "Transform");
+			const TManagedArray<int32>& GeometryToTransform = Collection->GetAttribute<int32>("TransformIndex", "Geometry");
+			const TManagedArray<FLinearColor>& BoneColor = Collection->GetAttribute<FLinearColor>("BoneColor", "Transform");
+		
+			// Add a sub item to the histogram for each of the geometry nodes in this GeometryCollection
+			for (int32 Index = 0; Index < NumElements; Index++)
+			{
+				if (Collection->Children[Index].Num() == 0)
+				{
+					TSharedRef<FGeometryCollectionHistogramItem> NewItem = MakeShared<FGeometryCollectionHistogramItem>(Guids[Index], Index, AsShared());
+
+					FLinearColor GeoColor = BoneColor[Index];
+					NewItem->SetColor(GeoColor);
+
+					NodesList.Add(NewItem);
+					NodesMap.Add(Guids[Index], NewItem);
+					GuidIndexMap.Add(Guids[Index], Index);
+				}
+			}
+		}
+	}
+
+	return NodesList;
+
+}
+
+void SGeometryCollectionHistogram::Construct(const FArguments& InArgs)
+{
+	BoneSelectionChangedDelegate = InArgs._OnBoneSelectionChanged;
+	bPerformingSelection = false;
+
+	ChildSlot
+	[
+		SAssignNew(ListView, SListView<FGeometryCollectionHistogramItemPtr>)
+		.ListItemsSource(&LeafNodes)
+		.OnSelectionChanged(this, &SGeometryCollectionHistogram::OnSelectionChanged)
+		.OnGenerateRow(this, &SGeometryCollectionHistogram::MakeHistogramRowWidget)
+	];
+}
+
+void SGeometryCollectionHistogram::SetComponents(const TArray<UGeometryCollectionComponent*>& InNewComponents)
+{
+	// Clear the cached Tree ItemSelection without affecting the SelectedBones as 
+	// we want to refresh the tree selection using selected bones
+	TGuardValue<bool> ExternalSelectionGuard(bPerformingSelection, true);
+	ListView->ClearSelection();
+	
+	RootNodes.Empty();
+	LeafNodes.Empty();
+
+	if (InNewComponents.Num() > 0)
+	{
+		for (UGeometryCollectionComponent* Component : InNewComponents)
+		{
+			RootNodes.Add(MakeShared<FGeometryCollectionHistogramItemComponent>(Component));
+			LeafNodes.Append(RootNodes.Last()->RegenerateNodes());
+		}
+
+		SetListIndices();
+
+		UHistogramSettings* HistogramSettings = GetMutableDefault<UHistogramSettings>();
+		InspectAttribute(HistogramSettings->InspectedAttribute);
+		NormalizeInspectedValues();
+	}
+	
+	UHistogramSettings* HistogramSettings = GetMutableDefault<UHistogramSettings>();
+	RefreshView(HistogramSettings->bSorted);
+}
+
+void SGeometryCollectionHistogram::InspectAttribute(EInspectedAttributeEnum InspectedAttribute)
+{
+	for (FGeometryCollectionHistogramItemPtr Node : LeafNodes)
+	{
+		Node->SetInspectedAttribute(InspectedAttribute);
+	}
+	NormalizeInspectedValues();
+
+	UHistogramSettings* HistogramSettings = GetMutableDefault<UHistogramSettings>();
+	RefreshView(HistogramSettings->bSorted);
+}
+
+void SGeometryCollectionHistogram::RefreshView(bool bSorted)
+{
+	
+	if (LeafNodes.Num() > 0)
+	{
+		if (bSorted)
+		{
+			// Sort nodes by value
+			LeafNodes.Sort([](const TSharedPtr<FGeometryCollectionHistogramItem> A, const TSharedPtr<FGeometryCollectionHistogramItem> B)
+				{
+					return A->GetInspectedValue() > B->GetInspectedValue(); 
+				});
+
+		}
+		else 
+		{
+			// Sort nodes by value
+			LeafNodes.Sort([](const TSharedPtr<FGeometryCollectionHistogramItem> A, const TSharedPtr<FGeometryCollectionHistogramItem> B)
+				{
+					return A->GetListIndex() > B->GetListIndex();
+				});
+		}	
+	}
+
+	ListView->RebuildList();
+}
+
+void SGeometryCollectionHistogram::RegenerateNodes()
+{
+	for (TSharedPtr<FGeometryCollectionHistogramItemComponent> Root : RootNodes)
+	{
+		Root->RegenerateNodes();
+	}
+
+	UHistogramSettings* HistogramSettings = GetMutableDefault<UHistogramSettings>();
+	RefreshView(HistogramSettings->bSorted);
+}
+
+
+void SGeometryCollectionHistogram::OnSelectionChanged(FGeometryCollectionHistogramItemPtr Item, ESelectInfo::Type SelectInfo)
+{
+	if (!bPerformingSelection && BoneSelectionChangedDelegate.IsBound())
+	{
+		TMap<UGeometryCollectionComponent*, TArray<int32>> ComponentToBoneSelectionMap;
+
+		ComponentToBoneSelectionMap.Reserve(RootNodes.Num());
+
+		// Create an entry for each component in the tree.  If the component has no selected bones then we return an empty array to signal that the selection should be cleared
+		for (const TSharedPtr<FGeometryCollectionHistogramItemComponent>& Root : RootNodes)
+		{
+			ComponentToBoneSelectionMap.Add(Root->GetComponent(), TArray<int32>());
+		}
+
+		if (!Item.IsValid())
+		{
+			ListView->ClearSelection();
+		}
+
+		FGeometryCollectionHistogramItemList SelectedItems;
+		ListView->GetSelectedItems(SelectedItems);
+
+		FScopedTransaction Transaction(FractureTransactionContexts::SelectBoneContext, LOCTEXT("SelectGeometryCollectionBoneTransaction", "Select Bone"), Item.IsValid() ? Item->GetComponent() : nullptr);
+		for (TSharedPtr<FGeometryCollectionHistogramItem> SelectedItem : SelectedItems)
+		{
+			if (SelectedItem->GetBoneIndex() != INDEX_NONE)
+			{
+				TArray<int32>& SelectedBones = ComponentToBoneSelectionMap.FindChecked(SelectedItem->GetComponent());
+				SelectedBones.Add(SelectedItem->GetBoneIndex());
+				SelectedItem->GetComponent()->Modify();
+			}
+		}
+
+		// Fire off the delegate for each component
+		for (TPair<UGeometryCollectionComponent*, TArray<int32>>& SelectionPair : ComponentToBoneSelectionMap)
+		{
+			BoneSelectionChangedDelegate.Execute(SelectionPair.Key, SelectionPair.Value);
+		}
+	}
+	
+}
+
+
+void SGeometryCollectionHistogram::NormalizeInspectedValues()
+{
+	// Sort nodes by value
+	LeafNodes.Sort([](const TSharedPtr<FGeometryCollectionHistogramItem> A, const TSharedPtr<FGeometryCollectionHistogramItem> B)
+		{
+			return A->GetInspectedValue() > B->GetInspectedValue();
+		});
+
+	// Normalize
+	float MaxInspectedValue = LeafNodes[0]->GetInspectedValue();
+	float MinInspectedValue = LeafNodes.Last()->GetInspectedValue();
+
+	// We allow a little extra room at the bottom so that bars don't drop to zero length
+	MinInspectedValue -= (MaxInspectedValue - MinInspectedValue) * 0.025;
+
+	for (TSharedPtr<FGeometryCollectionHistogramItem> LeafNode : LeafNodes)
+	{
+		LeafNode->SetNormalizedValue(MinInspectedValue, MaxInspectedValue);
+	}
+}
+
+void SGeometryCollectionHistogram::SetListIndices()
+{
+	int32 NodeCount = LeafNodes.Num();
+	for (int32 Index = 0; Index < NodeCount; ++Index)
+	{
+		LeafNodes[Index]->SetListIndex(Index);
+	}
+}
+
+
+void SGeometryCollectionHistogram::SetBoneSelection(UGeometryCollectionComponent* RootComponent, const TArray<int32>& InSelection, bool bClearCurrentSelection)
+{
+	TGuardValue<bool> ExternalSelectionGuard(bPerformingSelection, true);
+
+	if (bClearCurrentSelection)
+	{
+		ListView->ClearSelection();
+	}
+
+	bool bFirstSelection = true;
+
+	for (TSharedPtr<FGeometryCollectionHistogramItemComponent> RootNode : RootNodes)
+	{
+		if (RootNode->GetComponent() == RootComponent)
+		{
+			for (int32 BoneIndex : InSelection)
+			{
+				FGeometryCollectionHistogramItemPtr Item = RootNode->GetItemFromBoneIndex(BoneIndex);
+				if (Item.IsValid())
+				{
+					if (bFirstSelection)
+					{
+						ListView->RequestScrollIntoView(Item);
+						bFirstSelection = false;
+					}
+					ListView->SetItemSelection(Item, true);
+					ListView->SetItemHighlighted(Item, true);
+				}
+			}
+			break;
+		}
+	}
+	ListView->RequestListRefresh();
+
+}
+
+TSharedRef<ITableRow> SGeometryCollectionHistogram::MakeHistogramRowWidget(FGeometryCollectionHistogramItemPtr InItem, const TSharedRef<STableViewBase>& InOwnerTable)
+{
+	return InItem->MakeHistogramRowWidget(InOwnerTable);
+}
+
+
+#undef LOCTEXT_NAMESPACE
diff --git a/Engine/Plugins/Experimental/ChaosEditor/Source/FractureEditor/Private/SGeometryCollectionHistogram.h b/Engine/Plugins/Experimental/ChaosEditor/Source/FractureEditor/Private/SGeometryCollectionHistogram.h
new file mode 100644
--- /dev/null
+++ b/Engine/Plugins/Experimental/ChaosEditor/Source/FractureEditor/Private/SGeometryCollectionHistogram.h
@@ -0,0 +1,153 @@
+// Copyright Epic Games, Inc. All Rights Reserved.
+
+#pragma once
+
+#include "CoreMinimal.h"
+#include "Widgets/SCompoundWidget.h"
+#include "Math/Color.h"
+#include "Widgets/Views/SListView.h"
+
+#include "GeometryCollection/GeometryCollection.h"
+#include "GeometryCollection/GeometryCollectionComponent.h"
+
+#include "SGeometryCollectionHistogram.generated.h"
+
+
+
+struct FLinearColor;
+struct FVector2D;
+class FGeometryCollectionHistogramItem;
+class ITableRow;
+class STableViewBase;
+
+using FGeometryCollectionHistogramItemPtr = TSharedPtr<FGeometryCollectionHistogramItem>;
+using FGeometryCollectionHistogramItemList = TArray<TSharedPtr<FGeometryCollectionHistogramItem>>;
+
+
+UENUM(BlueprintType)
+enum class EInspectedAttributeEnum : uint8
+{
+	Volume = 0					UMETA(DisplayName = "Volume"),
+	Level = 1					UMETA(DisplayName = "Level"),
+	InitialDynamicState = 3		UMETA(DisplayName = "InitialDynamicState")
+};
+
+/** Settings for Histogram configuration. **/
+UCLASS()
+class UHistogramSettings : public UObject
+{
+
+	GENERATED_BODY()
+public:
+	UHistogramSettings(const FObjectInitializer& ObjInit);
+
+	/** What attribute are we inspecting? */
+	UPROPERTY(EditAnywhere, Category = HistogramSettings, meta = (DisplayName = "Inspected Attribute"))
+		EInspectedAttributeEnum InspectedAttribute;
+
+	/** Sort the values? */
+	UPROPERTY(EditAnywhere, Category = HistogramSettings, meta = (DisplayName = "Sort Values"))
+		bool bSorted;
+
+};
+
+class FGeometryCollectionHistogramItemComponent : public TSharedFromThis<FGeometryCollectionHistogramItemComponent>
+{
+public:
+	FGeometryCollectionHistogramItemComponent(UGeometryCollectionComponent* InComponent)
+		: Component(InComponent)
+	{
+	}
+
+	UGeometryCollectionComponent* GetComponent() const { return Component.IsValid() ? Component.Get() : nullptr; }
+
+	FGeometryCollectionHistogramItemPtr GetItemFromBoneIndex(int32 BoneIndex) const;
+	
+	FGeometryCollectionHistogramItemList RegenerateNodes();
+
+private:
+	TWeakObjectPtr<UGeometryCollectionComponent> Component;
+	TMap<FGuid, FGeometryCollectionHistogramItemPtr> NodesMap;
+	TMap<FGuid, int32> GuidIndexMap;
+};
+
+
+
+class FGeometryCollectionHistogramItem : public TSharedFromThis<FGeometryCollectionHistogramItem>
+{
+public:
+	FGeometryCollectionHistogramItem(const FGuid NewGuid, const int32 InBoneIndex, const TSharedPtr<FGeometryCollectionHistogramItemComponent> InParentComponentItem)
+		: Guid(NewGuid)
+		, BoneIndex(InBoneIndex)
+		, ParentComponentItem(InParentComponentItem)
+		, ListIndex(0)
+		, NodeColor(FLinearColor::Black)
+		, NormalizedValue(0.0)
+		, InspectedValue(0.0)
+		, HoverText(FText())
+	{}
+
+	TSharedRef<ITableRow> MakeHistogramRowWidget(const TSharedRef<STableViewBase>& InOwnerTable);
+	
+	UGeometryCollectionComponent* GetComponent() const { return ParentComponentItem->GetComponent(); }
+
+	float GetInspectedValue() const { return InspectedValue; }
+	void SetColor(const FLinearColor& InColor) { NodeColor = InColor; }
+	void SetInspectedAttribute(EInspectedAttributeEnum InspectedAttribute);
+	void SetNormalizedValue(float MinValue, float MaxValue);
+
+	int32 GetBoneIndex() const { return BoneIndex; }
+	void SetListIndex(int32 InListIndex) { ListIndex = InListIndex; }
+	int32 GetListIndex() const { return ListIndex; }
+
+private:
+	const FGuid Guid;
+	const int32 BoneIndex;
+	const TSharedPtr<FGeometryCollectionHistogramItemComponent> ParentComponentItem;
+	
+	int32 ListIndex;
+	FLinearColor NodeColor; 
+	float NormalizedValue;
+	float InspectedValue;
+	FText HoverText;
+};
+
+
+
+class SGeometryCollectionHistogram : public SCompoundWidget
+{
+	DECLARE_DELEGATE_TwoParams(FOnBoneSelectionChanged, UGeometryCollectionComponent*, TArray<int32>&);
+
+	SLATE_BEGIN_ARGS(SGeometryCollectionHistogram)
+	{}
+
+	SLATE_EVENT(FOnBoneSelectionChanged, OnBoneSelectionChanged)
+
+	SLATE_END_ARGS()
+
+public:
+	void Construct(const FArguments& InArgs);
+
+	TSharedRef<ITableRow> MakeHistogramRowWidget(FGeometryCollectionHistogramItemPtr InInfo, const TSharedRef<STableViewBase>& OwnerTable);
+	void SetComponents(const TArray<UGeometryCollectionComponent*>& InNewComponents);
+	void SetBoneSelection(UGeometryCollectionComponent* RootComponent, const TArray<int32>& InSelection, bool bClearCurrentSelection);
+	void InspectAttribute(EInspectedAttributeEnum InspectedAttribute);
+	void RefreshView(bool bSorted);
+	void RegenerateNodes();
+	void ClearSelection() { ListView->ClearSelection(); }
+	bool IsSelected() { return ListView->GetNumItemsSelected() > 0; }
+
+private:
+	void OnSelectionChanged(FGeometryCollectionHistogramItemPtr Item, ESelectInfo::Type SelectInfo);
+
+	void NormalizeInspectedValues();
+	void SetListIndices();
+
+private:
+	TSharedPtr<SListView<FGeometryCollectionHistogramItemPtr>> ListView;
+	TArray<TSharedPtr<FGeometryCollectionHistogramItemComponent>> RootNodes;
+	TArray<FGeometryCollectionHistogramItemPtr> LeafNodes;
+	
+	FOnBoneSelectionChanged BoneSelectionChangedDelegate;
+	bool bPerformingSelection;
+};
diff --git a/Engine/Plugins/Experimental/ChaosEditor/Source/FractureEditor/Private/SGeometryCollectionOutliner.cpp b/Engine/Plugins/Experimental/ChaosEditor/Source/FractureEditor/Private/SGeometryCollectionOutliner.cpp
--- a/Engine/Plugins/Experimental/ChaosEditor/Source/FractureEditor/Private/SGeometryCollectionOutliner.cpp
+++ b/Engine/Plugins/Experimental/ChaosEditor/Source/FractureEditor/Private/SGeometryCollectionOutliner.cpp
@@ -12,10 +12,16 @@
 #include "GeometryCollection/GeometryCollectionComponent.h"
 #include "FractureEditorMode.h"
 #include "ScopedTransaction.h"
+#include "FractureSettings.h"
 
 #define LOCTEXT_NAMESPACE "ChaosEditor"
 
 
+UOutlinerSettings::UOutlinerSettings(const FObjectInitializer& ObjInit)
+	: Super(ObjInit)
+	, ItemText(EOutlinerItemNameEnum::BoneIndex)
+{}
+
 void SGeometryCollectionOutliner::Construct(const FArguments& InArgs)
 {
 	BoneSelectionChangedDelegate = InArgs._OnBoneSelectionChanged;
@@ -33,6 +39,11 @@ void SGeometryCollectionOutliner::Construct(const FArguments& InArgs)
 	];
 }
 
+void SGeometryCollectionOutliner::RegenerateItems()
+{
+	TreeView->RebuildList();
+}
+
 TSharedRef<ITableRow> SGeometryCollectionOutliner::MakeTreeRowWidget(FGeometryCollectionTreeItemPtr InItem, const TSharedRef<STableViewBase>& InOwnerTable)
 {
 	return InItem->MakeTreeRowWidget(InOwnerTable);
@@ -143,6 +154,23 @@ void SGeometryCollectionOutliner::ExpandRecursive(TSharedPtr<FGeometryCollection
 	}
 }
 
+void SGeometryCollectionOutliner::SetHistogramSelection(UGeometryCollectionComponent* RootComponent, TArray<int32>& SelectedBones)
+{
+	// Find the matching component
+	for (TSharedPtr<FGeometryCollectionTreeItemComponent> RootNode : RootNodes)
+	{
+		if (RootNode->GetComponent() == RootComponent)
+		{
+			// Copy the histogram selection.
+			RootNode->SetHistogramSelection(SelectedBones);
+			RootNode->RegenerateChildren();
+			TreeView->RequestTreeRefresh();
+			ExpandAll();
+			return;
+		}
+	}
+}
+
 void SGeometryCollectionOutliner::SetBoneSelection(UGeometryCollectionComponent* RootComponent, const TArray<int32>& InSelection, bool bClearCurrentSelection)
 {
 	TGuardValue<bool> ExternalSelectionGuard(bPerformingSelection, true);
@@ -328,26 +356,85 @@ void FGeometryCollectionTreeItemComponent::RegenerateChildren()
 			// Add a sub item to the outliner tree for each of the bones/chunks in this GeometryCollection
 			for (int32 Index = 0; Index < NumElements; Index++)
 			{
-				TSharedRef<FGeometryCollectionTreeItemBone> NewItem = MakeShared<FGeometryCollectionTreeItemBone>(Guids[Index], Index, *this);
-				if (Parents[Index] == RootIndex)
+				if (FilterBoneIndex(Index))
 				{
-					// The actual children directly beneath this node are the ones without a parent.  The rest are children of children
-					MyChildren.Add(NewItem);
-				}
+					TSharedRef<FGeometryCollectionTreeItemBone> NewItem = MakeShared<FGeometryCollectionTreeItemBone>(Guids[Index], Index, *this);
+					if (Parents[Index] == RootIndex)
+					{
+						// The actual children directly beneath this node are the ones without a parent.  The rest are children of children
+						MyChildren.Add(NewItem);
+					}
 
-				NodesMap.Add(Guids[Index], NewItem);
-				GuidIndexMap.Add(Guids[Index], Index);
+					NodesMap.Add(Guids[Index], NewItem);
+					GuidIndexMap.Add(Guids[Index], Index);
+				}			
 			}
+
 		}
 	}
 }
 
+void FGeometryCollectionTreeItemComponent::SetHistogramSelection(TArray<int32>& SelectedBones)
+{
+	HistogramSelection = SelectedBones;
+}
+
+bool FGeometryCollectionTreeItemComponent::FilterBoneIndex(int32 BoneIndex) const
+{
+	FGeometryCollection* Collection = Component->GetRestCollection()->GetGeometryCollection().Get();
+	TManagedArray<TSet<int32>>& Children = Collection->GetAttribute<TSet<int32>>("Children", FGeometryCollection::TransformGroup);
+
+	if (Children[BoneIndex].Num() == 0)
+	{
+		// We don't display leaf nodes deeper than the view level.
+		UFractureSettings* FractureSettings = GetMutableDefault<UFractureSettings>();
+
+		if (FractureSettings->FractureLevel >= 0)
+		{
+			TManagedArray<int32>& Level = Collection->GetAttribute<int32>("Level", FTransformCollection::TransformGroup);
+			if (Level[BoneIndex] != FractureSettings->FractureLevel)
+			{
+				return false;
+			}
+		}
+
+		// If anything is selected int the Histogram, we filter by that selection.
+		if (HistogramSelection.Num() > 0)
+		{
+			if (!HistogramSelection.Contains(BoneIndex))
+			{
+				return false;
+			}
+		}		
+	}
+
+	return true;	
+}
+
 TSharedRef<ITableRow> FGeometryCollectionTreeItemBone::MakeTreeRowWidget(const TSharedRef<STableViewBase>& InOwnerTable)
 {
+	UOutlinerSettings* OutlinerSettings = GetMutableDefault<UOutlinerSettings>();
+	FText ItemText = ParentComponentItem->GetDisplayNameForBone(Guid);
+	if (OutlinerSettings->ItemText == EOutlinerItemNameEnum::BoneIndex)
+	{
+		ItemText = FText::FromString(FString::FromInt(GetBoneIndex()));
+	}
+
+	// Cyan indicates cluster, Gray is leaf node	
+	TSharedPtr<FGeometryCollection, ESPMode::ThreadSafe> GeometryCollectionPtr = ParentComponentItem->GetComponent()->GetRestCollection()->GetGeometryCollection();
+	const TManagedArray<TSet<int32>>& Children = GeometryCollectionPtr->GetAttribute<TSet<int32>>("Children", FGeometryCollection::TransformGroup);
+	FSlateColor TextColor(FLinearColor::Gray);
+	if (Children[GetBoneIndex()].Num() > 0)
+	{
+		TextColor = FColor::Cyan;
+	}
+	
 	return SNew(STableRow<FGeometryCollectionTreeItemPtr>, InOwnerTable)
 		.Content()
-		[
-			SNew(STextBlock).Text(ParentComponentItem->GetDisplayNameForBone(Guid))
+		[			
+			SNew(STextBlock)
+			.Text(ItemText)
+			.ColorAndOpacity(TextColor)
 		];
 }
 
diff --git a/Engine/Plugins/Experimental/ChaosEditor/Source/FractureEditor/Private/SGeometryCollectionOutliner.h b/Engine/Plugins/Experimental/ChaosEditor/Source/FractureEditor/Private/SGeometryCollectionOutliner.h
--- a/Engine/Plugins/Experimental/ChaosEditor/Source/FractureEditor/Private/SGeometryCollectionOutliner.h
+++ b/Engine/Plugins/Experimental/ChaosEditor/Source/FractureEditor/Private/SGeometryCollectionOutliner.h
@@ -13,18 +13,41 @@
 #include "Templates/SharedPointer.h"
 #include "UObject/ObjectKey.h"
 
-// #include "GeometryCollection/GeometryCollectionComponent.h"
-// #include "GeometryCollection/GeometryCollectionActor.h"
 #include "GeometryCollection/GeometryCollection.h"
+#include "GeometryCollection/GeometryCollectionComponent.h"
+
+#include "SGeometryCollectionOutliner.generated.h"
 
 class FGeometryCollection;
 class FGeometryCollectionTreeItem;
 class FGeometryCollectionTreeItemBone;
-class UGeometryCollectionComponent;
 
 typedef TArray<TSharedPtr<FGeometryCollectionTreeItem>> FGeometryCollectionTreeItemList;
 typedef TSharedPtr<FGeometryCollectionTreeItem> FGeometryCollectionTreeItemPtr;
 
+
+UENUM(BlueprintType)
+enum class EOutlinerItemNameEnum : uint8
+{
+	BoneName = 0					UMETA(DisplayName = "Bone Name"),
+	BoneIndex = 1					UMETA(DisplayName = "Bone Index"),
+};
+
+/** Settings for Outliner configuration. **/
+UCLASS()
+class UOutlinerSettings : public UObject
+{
+
+	GENERATED_BODY()
+public:
+	UOutlinerSettings(const FObjectInitializer& ObjInit);
+
+	/** What is displayed in Outliner text */
+	UPROPERTY(EditAnywhere, Category = OutlinerSettings, meta = (DisplayName = "Item Text"))
+	EOutlinerItemNameEnum ItemText;
+};
+
+
 class FGeometryCollectionTreeItem : public TSharedFromThis<FGeometryCollectionTreeItem>
 {
 public:
@@ -59,6 +82,11 @@ public:
 	void ExpandAll(TSharedPtr<STreeView<FGeometryCollectionTreeItemPtr>> TreeView);
 	void RegenerateChildren();
 
+	void SetHistogramSelection(TArray<int32>& SelectedBones);
+
+private:
+	bool FilterBoneIndex(int32 BoneIndex) const;
+
 private:
 	TWeakObjectPtr<UGeometryCollectionComponent> Component;
 
@@ -69,6 +97,8 @@ private:
 	TMap<FGuid, int32> GuidIndexMap;
 	FGuid RootGuid;
 	int32 RootIndex;
+
+	TArray<int32> HistogramSelection;
 };
 
 class FGeometryCollectionTreeItemBone : public FGeometryCollectionTreeItem
@@ -112,6 +142,8 @@ class SGeometryCollectionOutliner: public SCompoundWidget
 public:
 	void Construct(const FArguments& InArgs);
 
+	void RegenerateItems();
+
 	TSharedRef<ITableRow> MakeTreeRowWidget(FGeometryCollectionTreeItemPtr InInfo, const TSharedRef<STableViewBase>& OwnerTable);
 	void OnGetChildren(TSharedPtr<FGeometryCollectionTreeItem> InInfo, TArray< TSharedPtr<FGeometryCollectionTreeItem> >& OutChildren);
 
@@ -122,6 +154,9 @@ public:
 	void ExpandAll();
 	void ExpandRecursive(TSharedPtr<FGeometryCollectionTreeItem> TreeItem, bool bInExpansionState) const;
 
+	// Set the histogram filter on the component matching RootComponent.
+	void SetHistogramSelection(UGeometryCollectionComponent* RootComponent, TArray<int32>& SelectedBones);
+
 private:
 	// void GetItemParentsForGuid(int32 Guid);
 	//void GetItemChildrenForGuid(int32 Guid);
diff --git a/Engine/Plugins/Experimental/ChaosEditor/Source/FractureEditor/Public/FractureEditorModeToolkit.h b/Engine/Plugins/Experimental/ChaosEditor/Source/FractureEditor/Public/FractureEditorModeToolkit.h
--- a/Engine/Plugins/Experimental/ChaosEditor/Source/FractureEditor/Public/FractureEditorModeToolkit.h
+++ b/Engine/Plugins/Experimental/ChaosEditor/Source/FractureEditor/Public/FractureEditorModeToolkit.h
@@ -12,9 +12,10 @@
 class IDetailsView;
 class IPropertyHandle;
 class SScrollBox;
-struct FFractureToolContext;
+class FFractureToolContext;
 class FGeometryCollection;
 class SGeometryCollectionOutliner;
+class SGeometryCollectionHistogram;
 class AGeometryCollectionActor;
 class UGeometryCollectionComponent;
 class UGeometryCollection;
@@ -27,6 +28,12 @@ namespace GeometryCollection
 enum class ESelectionMode: uint8;
 }
 
+namespace Chaos
+{
+	template<class T, int d>
+	class TParticles;
+}
+
 class FFractureViewSettingsCustomization : public IDetailCustomization
 {
 public:
@@ -35,14 +42,40 @@ public:
 
 	virtual void CustomizeDetails(IDetailLayoutBuilder& DetailBuilder) override;
 
-	private:
-		FFractureEditorModeToolkit* Toolkit;
+private:
+	FFractureEditorModeToolkit* Toolkit;
+};
+
+class FHistogramSettingsCustomization : public IDetailCustomization
+{
+public:
+	FHistogramSettingsCustomization(FFractureEditorModeToolkit* FractureToolkit);
+	static TSharedRef<IDetailCustomization> MakeInstance(FFractureEditorModeToolkit* FractureToolkit);
+
+	virtual void CustomizeDetails(IDetailLayoutBuilder& DetailBuilder) override;
+
+private:
+	FFractureEditorModeToolkit* Toolkit;
+};
+
+class FOutlinerSettingsCustomization : public IDetailCustomization
+{
+public:
+	FOutlinerSettingsCustomization(FFractureEditorModeToolkit* FractureToolkit);
+	static TSharedRef<IDetailCustomization> MakeInstance(FFractureEditorModeToolkit* FractureToolkit);
+
+	virtual void CustomizeDetails(IDetailLayoutBuilder& DetailBuilder) override;
+
+private:
+	FFractureEditorModeToolkit* Toolkit;
 };
 
 class FRACTUREEDITOR_API FFractureEditorModeToolkit : public FModeToolkit, public FGCObject
 {
 public:
 
+	using FGeometryCollectionPtr = TSharedPtr<FGeometryCollection, ESPMode::ThreadSafe>;
+
 	FFractureEditorModeToolkit();
 	~FFractureEditorModeToolkit();
 	
@@ -93,6 +126,10 @@ public:
 	FReply OnModalClicked();
 	bool CanExecuteModal() const;
 
+	// Filter callbacks
+	FReply ResetHistogramSelection();
+	bool CanResetFilter() const;
+
 	static void GetSelectedGeometryCollectionComponents(TSet<UGeometryCollectionComponent*>& GeomCompSelection);
 
 	static void AddAdditionalAttributesIfRequired(UGeometryCollection* GeometryCollectionObject);
@@ -115,6 +152,8 @@ public:
 
 	void UpdateExplodedVectors(UGeometryCollectionComponent* GeometryCollectionComponent) const;
 
+	void RegenerateHistogram();
+
 	TSharedPtr<SWidget> ExplodedViewWidget;
 	TSharedPtr<SWidget> LevelViewWidget;
 	TSharedPtr<SWidget> ShowBoneColorsWidget;
@@ -124,7 +163,11 @@ protected:
 	static bool IsSelectedActorsInEditorWorld();	
 
 private:
+	static void UpdateGeometryComponentAttributes(UGeometryCollectionComponent* Component);
+	static void UpdateVolumes(FGeometryCollectionPtr GeometryCollection, const Chaos::TParticles<float, 3>& MassSpaceParticles, int32 TransformIndex);
+
 	void OnOutlinerBoneSelectionChanged(UGeometryCollectionComponent* RootComponent, TArray<int32>& SelectedBones);
+	void OnHistogramBoneSelectionChanged(UGeometryCollectionComponent* RootComponent, TArray<int32>& SelectedBones);
 	void BindCommands();
 
 private:
@@ -132,7 +175,10 @@ private:
 
 	TSharedPtr<IDetailsView> DetailsView;
 	TSharedPtr<IDetailsView> ViewSettingsDetailsView;
+	TSharedPtr<IDetailsView> HistogramDetailsView;
+	TSharedPtr<IDetailsView> OutlinerDetailsView;
 	TSharedPtr<SWidget> ToolkitWidget;
 	TSharedPtr<SGeometryCollectionOutliner> OutlinerView;
+	TSharedPtr<SGeometryCollectionHistogram> HistogramView;
 
 };
diff --git a/Engine/Plugins/Experimental/ChaosEditor/Source/FractureEditor/Public/FractureTool.cpp b/Engine/Plugins/Experimental/ChaosEditor/Source/FractureEditor/Public/FractureTool.cpp
--- a/Engine/Plugins/Experimental/ChaosEditor/Source/FractureEditor/Public/FractureTool.cpp
+++ b/Engine/Plugins/Experimental/ChaosEditor/Source/FractureEditor/Public/FractureTool.cpp
@@ -5,6 +5,9 @@
 #include "Editor.h"
 #include "Engine/Selection.h"
 
+#include "FractureToolContext.h"
+#include "FractureSelectionTools.h"
+
 #include "GeometryCollection/GeometryCollectionObject.h"
 #include "GeometryCollection/GeometryCollection.h"
 #include "GeometryCollection/GeometryCollectionComponent.h"
@@ -124,65 +127,108 @@ void UFractureActionTool::GetSelectedGeometryCollectionComponents(TSet<UGeometry
 	}
 }
 
+void UFractureActionTool::Refresh(FFractureToolContext& Context, FFractureEditorModeToolkit* Toolkit)
+{
+	UGeometryCollectionComponent* GeometryCollectionComponent = Context.GetGeometryCollectionComponent();
+	TSharedPtr<FGeometryCollection, ESPMode::ThreadSafe> GeometryCollectionPtr = Context.GetGeometryCollection();
+	
+	FGeometryCollectionClusteringUtility::UpdateHierarchyLevelOfChildren(GeometryCollectionPtr.Get(), -1);
 
-void UFractureModalTool::Execute(TWeakPtr<FFractureEditorModeToolkit> InToolkit)
+	FScopedColorEdit EditBoneColor(GeometryCollectionComponent, true);
+	EditBoneColor.SetSelectedBones(Context.GetSelection());
+
+	FFractureSelectionTools::ToggleSelectedBones(GeometryCollectionComponent, Context.GetSelection(), true);
+	Toolkit->SetBoneSelection(GeometryCollectionComponent, Context.GetSelection(), true);
+
+	Toolkit->UpdateExplodedVectors(GeometryCollectionComponent);
+
+	GeometryCollectionComponent->MarkRenderDynamicDataDirty();
+	GeometryCollectionComponent->MarkRenderStateDirty();
+}
+
+void UFractureActionTool::SetOutlinerComponents(TArray<FFractureToolContext>& InContexts, FFractureEditorModeToolkit* Toolkit)
 {
-	if (!InToolkit.IsValid())
+	// Extract components from contexts
+	TArray<UGeometryCollectionComponent*> Components;
+	for (FFractureToolContext& Context : InContexts)
 	{
-		return;
+		Components.Add(Context.GetGeometryCollectionComponent());
 	}
+	Toolkit->SetOutlinerComponents(Components);
+}
 
-	TArray<FFractureToolContext> FractureContexts;
-	GetFractureContexts(FractureContexts);
+TArray<FFractureToolContext> UFractureActionTool::GetFractureToolContexts() const
+{
+	TArray<FFractureToolContext> Contexts;
 
-	TArray<UGeometryCollectionComponent*> NewComponents;
-	for (FFractureToolContext& FractureContext : FractureContexts)
-	{	
-		FractureContext.FracturedGeometryCollection->Modify();
+	// A context is gathered for each selected GeometryCollection component.
+	TSet<UGeometryCollectionComponent*> GeomCompSelection;
+	GetSelectedGeometryCollectionComponents(GeomCompSelection);
+	for (UGeometryCollectionComponent* GeometryCollectionComponent : GeomCompSelection)
+	{
+		Contexts.Emplace(GeometryCollectionComponent);
+	}
+
+	return Contexts;
+}
+
+
+void UFractureModalTool::Execute(TWeakPtr<FFractureEditorModeToolkit> InToolkit)
+{
+	if (InToolkit.IsValid())
+	{
+		FFractureEditorModeToolkit* Toolkit = InToolkit.Pin().Get();
 
-		ExecuteFracture(FractureContext);
+		TArray<FFractureToolContext> FractureContexts = GetFractureToolContexts();
 
+		for (FFractureToolContext& FractureContext : FractureContexts)
 		{
-			TSharedPtr<FGeometryCollection, ESPMode::ThreadSafe> GeometryCollectionPtr = FractureContext.FracturedGeometryCollection->GetGeometryCollection();
-			FGeometryCollection* OutGeometryCollection = GeometryCollectionPtr.Get();
-			FGeometryCollectionClusteringUtility::UpdateHierarchyLevelOfChildren(OutGeometryCollection, -1);
+			FractureContext.GetFracturedGeometryCollection()->Modify();
+
+			int32 FirstNewGeometryIndex = ExecuteFracture(FractureContext);
+			
+			if (FirstNewGeometryIndex > INDEX_NONE)
+			{
+				// Based on the first new geometry index, generate a list of new transforms generated by the fracture.
+				const TManagedArray<int32>& TransformIndex = FractureContext.GetGeometryCollection()->GetAttribute<int32>("TransformIndex", FGeometryCollection::GeometryGroup);
+				int32 LastGeometryIndex = TransformIndex.Num();
+
+				TArray<int32> NewTransforms;
+				NewTransforms.Reserve(LastGeometryIndex - FirstNewGeometryIndex);
+				for (int32 GeometryIndex = FirstNewGeometryIndex; GeometryIndex < LastGeometryIndex; ++GeometryIndex)
+				{
+					NewTransforms.Add(TransformIndex[GeometryIndex]);
+				}
+
+				FractureContext.SetSelection(NewTransforms);
+
+				Toolkit->RegenerateHistogram();
+			}
+			
+
+			FGeometryCollectionClusteringUtility::UpdateHierarchyLevelOfChildren(FractureContext.GetGeometryCollection().Get(), -1);
 
 			// Update Nanite resource data to correctly reflect modified geometry collection data
 			{
-				FractureContext.FracturedGeometryCollection->ReleaseResources();
+				FractureContext.GetFracturedGeometryCollection()->ReleaseResources();
 
-				if (FractureContext.FracturedGeometryCollection->EnableNanite)
+				if (FractureContext.GetFracturedGeometryCollection()->EnableNanite)
 				{
-					FractureContext.FracturedGeometryCollection->NaniteData = UGeometryCollection::CreateNaniteData(OutGeometryCollection);
+					FractureContext.GetFracturedGeometryCollection()->NaniteData = UGeometryCollection::CreateNaniteData(FractureContext.GetGeometryCollection().Get());
 				}
 				else
 				{
-					FractureContext.FracturedGeometryCollection->NaniteData = MakeUnique<FGeometryCollectionNaniteData>();
+					FractureContext.GetFracturedGeometryCollection()->NaniteData = MakeUnique<FGeometryCollectionNaniteData>();
 				}
 
-				FractureContext.FracturedGeometryCollection->InitResources();
+				FractureContext.GetFracturedGeometryCollection()->InitResources();
 			}
-		}
-
-		UGeometryCollectionComponent* GeometryCollectionComponent = Cast<UGeometryCollectionComponent>(FractureContext.OriginalPrimitiveComponent);
-		NewComponents.AddUnique(GeometryCollectionComponent);
-		FGeometryCollectionEdit GCEdit = GeometryCollectionComponent->EditRestCollection();
-		UGeometryCollection* GCObject = GCEdit.GetRestCollection();
-		TSharedPtr<FGeometryCollection, ESPMode::ThreadSafe> GeometryCollectionPtr = GCObject->GetGeometryCollection();
-		FGeometryCollectionClusteringUtility::UpdateHierarchyLevelOfChildren(GeometryCollectionPtr.Get(), -1);
-
-		FScopedColorEdit EditBoneColor(GeometryCollectionComponent, true);
-		EditBoneColor.SelectBones(GeometryCollection::ESelectionMode::None);
-		InToolkit.Pin()->SetBoneSelection(GeometryCollectionComponent, EditBoneColor.GetSelectedBones(), true);
 
-		InToolkit.Pin()->UpdateExplodedVectors(GeometryCollectionComponent);
+			Refresh(FractureContext, Toolkit);
+		}
 
-		GeometryCollectionComponent->MarkRenderDynamicDataDirty();
-		GeometryCollectionComponent->MarkRenderStateDirty();
-		
+		SetOutlinerComponents(FractureContexts, Toolkit);
 	}
-
-	InToolkit.Pin()->SetOutlinerComponents(NewComponents);
 }
 
 bool UFractureModalTool::CanExecute() const
@@ -205,107 +251,3 @@ void UFractureModalTool::PostEditChangeChainProperty(struct FPropertyChangedChai
 	FractureContextChanged();
 }
 
-void UFractureModalTool::GetFractureContexts(TArray<FFractureToolContext>& FractureContexts) const
-{
-	TSet<UGeometryCollectionComponent*> GeometryCollectionComponents;
-	GetSelectedGeometryCollectionComponents(GeometryCollectionComponents);
-
-	for (UGeometryCollectionComponent* GeometryCollectionComponent : GeometryCollectionComponents)
-	{
-		FGeometryCollectionEdit RestCollection = GeometryCollectionComponent->EditRestCollection(GeometryCollection::EEditUpdate::None);
-		UGeometryCollection* FracturedGeometryCollection = RestCollection.GetRestCollection();
-		if (FracturedGeometryCollection == nullptr)
-		{
-			continue;
-		}
-
-		TSharedPtr<FGeometryCollection, ESPMode::ThreadSafe> GeometryCollectionPtr = FracturedGeometryCollection->GetGeometryCollection();
-		FGeometryCollection* OutGeometryCollection = GeometryCollectionPtr.Get();
-
-		TArray<int32> FilteredBones = FilterBones(GeometryCollectionComponent->GetSelectedBones(), OutGeometryCollection);
-
-
-		const TManagedArray<FTransform>& Transform = OutGeometryCollection->GetAttribute<FTransform>("Transform", FGeometryCollection::TransformGroup);
-		const TManagedArray<int32>& TransformToGeometryIndex = OutGeometryCollection->GetAttribute<int32>("TransformToGeometryIndex", FGeometryCollection::TransformGroup);
-		const TManagedArray<FBox>& BoundingBoxes = OutGeometryCollection->GetAttribute<FBox>("BoundingBox", FGeometryCollection::GeometryGroup);
-
-		TArray<FTransform> Transforms;
-		GeometryCollectionAlgo::GlobalMatrices(Transform, OutGeometryCollection->Parent, Transforms);
-
-
-		TMap<int32, FBox> BoundsToBone;
-		for (int32 Idx = 0, ni = FracturedGeometryCollection->NumElements(FGeometryCollection::TransformGroup); Idx < ni; ++Idx)
-		{
-			if (TransformToGeometryIndex[Idx] > -1)
-			{
-				ensure(TransformToGeometryIndex[Idx] > -1);
-				BoundsToBone.Add(Idx, BoundingBoxes[TransformToGeometryIndex[Idx]].TransformBy(Transforms[Idx]));
-			}
-		}
-
-		GenerateFractureToolContext(
-			GeometryCollectionComponent->GetOwner(),
-			GeometryCollectionComponent,
-			FracturedGeometryCollection,
-			FilteredBones,
-			BoundsToBone,
-			TransformToGeometryIndex,
-			-1,
-			FractureContexts);
-
-	}
-}
-
-TArray<int32> UFractureModalTool::FilterBones(const TArray<int32>& SelectedBonesOriginal, const FGeometryCollection* const GeometryCollection) const
-{
-	// Keep only leaf nodes
-	const TManagedArray<TSet<int32>>& Children = GeometryCollection->GetAttribute<TSet<int32>>("Children", FGeometryCollection::TransformGroup);
-
-	TArray<int32> SelectedBones;
-	SelectedBones.Reserve(SelectedBonesOriginal.Num());
-	for (int32 BoneIndex : SelectedBonesOriginal)
-	{
-		if (Children[BoneIndex].Num() == 0)
-		{
-			SelectedBones.Add(BoneIndex);
-		}
-	}
-
-	return SelectedBones;
-}
-
-void UFractureModalTool::GenerateFractureToolContext(
-	AActor* InActor, 
-	UGeometryCollectionComponent* InComponent, 
-	UGeometryCollection* InGeometryCollection, 
-	const TArray<int32>& InSelectedBones, 
-	TMap<int32, FBox>& InBoundsToBone,
-	const TManagedArray<int32>& TransformToGeometryIndex,
-	int32 RandomSeed, 
-	TArray<FFractureToolContext>& OutFractureContexts
-) const
-{
-	OutFractureContexts.AddDefaulted();
-	FFractureToolContext& FractureContext = OutFractureContexts.Last();
-	FractureContext.RandomSeed = FMath::Rand();
-	if (RandomSeed > -1)
-	{
-		// make sure it's unique for each context if it's specified.
-		FractureContext.RandomSeed = RandomSeed + OutFractureContexts.Num();
-	}
-
-	FractureContext.OriginalActor = InActor;
-	FractureContext.Transform = InActor->GetActorTransform();
-	FractureContext.OriginalPrimitiveComponent = InComponent;
-	FractureContext.FracturedGeometryCollection = InGeometryCollection;
-	FractureContext.SelectedBones = InSelectedBones;
-
-	FractureContext.Bounds = FBox(ForceInit);
-	for (int32 BoneIndex : FractureContext.SelectedBones)
-	{
-		if (TransformToGeometryIndex[BoneIndex] > INDEX_NONE)
-		{
-			FractureContext.Bounds += InBoundsToBone[BoneIndex];
-		}
-	}
-}
diff --git a/Engine/Plugins/Experimental/ChaosEditor/Source/FractureEditor/Public/FractureTool.h b/Engine/Plugins/Experimental/ChaosEditor/Source/FractureEditor/Public/FractureTool.h
--- a/Engine/Plugins/Experimental/ChaosEditor/Source/FractureEditor/Public/FractureTool.h
+++ b/Engine/Plugins/Experimental/ChaosEditor/Source/FractureEditor/Public/FractureTool.h
@@ -14,6 +14,7 @@
 
 class UGeometryCollection;
 class UFractureModalTool;
+class FFractureToolContext;
 
 template <typename T>
 class TManagedArray;
@@ -35,19 +36,6 @@ public:
 };
 
 
-struct FFractureToolContext
-{
-	AActor* OriginalActor;
-	UPrimitiveComponent* OriginalPrimitiveComponent;
-	UGeometryCollection* FracturedGeometryCollection;
-	FString ParentName;
-	FTransform Transform;
-	FBox Bounds;
-	int32 RandomSeed;
-	TArray<int32> SelectedBones;
-};
-
-
 /** Tools derived from this class should require parameter inputs from the user, only the bone selection. */
 UCLASS(Abstract)
 class UFractureActionTool : public UObject
@@ -72,8 +60,7 @@ public:
 
 	virtual void RegisterUICommand(FFractureEditorCommands* BindingContext) {}
 
-	virtual void GetFractureToolContexts(TArray<FFractureToolContext>& FractureToolContexts) const {}
-
+	virtual TArray<FFractureToolContext> GetFractureToolContexts() const;
 	
 protected:
 	static bool IsStaticMeshSelected();
@@ -81,6 +68,9 @@ protected:
 	static void AddSingleRootNodeIfRequired(UGeometryCollection* GeometryCollectionObject);
 	static void AddAdditionalAttributesIfRequired(UGeometryCollection* GeometryCollectionObject);
 	static void GetSelectedGeometryCollectionComponents(TSet<UGeometryCollectionComponent*>& GeomCompSelection);
+	static void Refresh(FFractureToolContext& Context, FFractureEditorModeToolkit* Toolkit);
+	static void SetOutlinerComponents(TArray<FFractureToolContext>& InContexts, FFractureEditorModeToolkit* Toolkit);
+
 
 protected:
 	TSharedPtr<FUICommandInfo> UICommandInfo;
@@ -107,7 +97,8 @@ public:
 	virtual void Execute(TWeakPtr<FFractureEditorModeToolkit> InToolkit) override;
 	virtual bool CanExecute() const override;
 
-	virtual void ExecuteFracture(const FFractureToolContext& FractureContext) {}
+	/** Executes function that generates new geometry. Returns the first new geometry index. */
+	virtual int32 ExecuteFracture(const FFractureToolContext& FractureContext) { return INDEX_NONE; }
 
 	/** Draw callback from edmode*/
 	virtual void Render(const FSceneView* View, FViewport* Viewport, FPrimitiveDrawInterface* PDI) {}
@@ -115,22 +106,6 @@ public:
 	virtual void PostEditChangeChainProperty(struct FPropertyChangedChainEvent& PropertyChangedEvent) override;
 	virtual void FractureContextChanged() {}
 
-protected:
-
-	virtual void GetFractureContexts(TArray<FFractureToolContext>& FractureContexts) const;
-	virtual TArray<int32> FilterBones(const TArray<int32>& SelectedBonesOriginal, const FGeometryCollection* const GeometryCollection) const;
-
-	void GenerateFractureToolContext(
-		AActor* InActor,
-		UGeometryCollectionComponent* InComponent,
-		UGeometryCollection* InGeometryCollection,
-		const TArray<int32>& InSelectedBones,
-		TMap<int32, FBox>& InBoundsToBone,
-		const TManagedArray<int32>& TransformToGeometryIndex,
-		int32 RandomSeed,
-		TArray<FFractureToolContext>& OutFractureContexts
-	) const;
-
 };
 
 
diff --git a/Engine/Source/Runtime/Experimental/Chaos/Private/GeometryCollection/GeometryCollectionClusteringUtility.cpp b/Engine/Source/Runtime/Experimental/Chaos/Private/GeometryCollection/GeometryCollectionClusteringUtility.cpp
--- a/Engine/Source/Runtime/Experimental/Chaos/Private/GeometryCollection/GeometryCollectionClusteringUtility.cpp
+++ b/Engine/Source/Runtime/Experimental/Chaos/Private/GeometryCollection/GeometryCollectionClusteringUtility.cpp
@@ -23,7 +23,7 @@ void FGeometryCollectionClusteringUtility::ClusterBonesUnderNewNode(FGeometryCol
 	int32 SourceBoneIndex = InsertAtIndex;
 	int32 OriginalParentIndex = Parents[SourceBoneIndex];
 	BoneNames[NewBoneIndex] = BoneNames[SourceBoneIndex];
-	Parents[NewBoneIndex] = Parents[SourceBoneIndex];
+	Parents[NewBoneIndex] = OriginalParentIndex;
 	Children[NewBoneIndex] = TSet<int32>(SelectedBones);
 
 	Transforms[NewBoneIndex] = FTransform::Identity;
@@ -148,7 +148,7 @@ void FGeometryCollectionClusteringUtility::ClusterAllBonesUnderNewRoot(FGeometry
 }
 
 
-void FGeometryCollectionClusteringUtility::ClusterBonesUnderExistingRoot(FGeometryCollection* GeometryCollection, TArray<int32>& SourceElements)
+void FGeometryCollectionClusteringUtility::ClusterBonesUnderExistingRoot(FGeometryCollection* GeometryCollection, const TArray<int32>& SourceElements)
 {
 	check(GeometryCollection);
 	bool CalcNewLocalTransform = true;
@@ -224,8 +224,10 @@ void FGeometryCollectionClusteringUtility::ClusterBonesUnderExistingNode(FGeomet
 	TManagedArray<TSet<int32>>& Children = GeometryCollection->Children;
 	TManagedArray<FTransform>& Transforms = GeometryCollection->Transform;
 	TManagedArray<FString>& BoneNames = GeometryCollection->BoneName;
-	TManagedArray<FTransform>& ExplodedTransforms = GeometryCollection->GetAttribute<FTransform>("ExplodedTransform", FGeometryCollection::TransformGroup);
-	TManagedArray<FVector>& ExplodedVectors = GeometryCollection->GetAttribute<FVector>("ExplodedVector", FGeometryCollection::TransformGroup);
+
+	// These attributes are apparently deprecated?
+	//TManagedArray<FTransform>& ExplodedTransforms = GeometryCollection->GetAttribute<FTransform>("ExplodedTransform", FGeometryCollection::TransformGroup);
+	//TManagedArray<FVector>& ExplodedVectors = GeometryCollection->GetAttribute<FVector>("ExplodedVector", FGeometryCollection::TransformGroup);
 
 	// remove Merge Node if it's in the list - happens due to the way selection works
 	TArray<int32> SourceElements;
@@ -258,17 +260,17 @@ void FGeometryCollectionClusteringUtility::ClusterBonesUnderExistingNode(FGeomet
 				ParentsToUpdateNames.AddUnique(Parents[SourceElement]);
 			}
 
-			ResetSliderTransforms(ExplodedTransforms, Transforms);
+			//ResetSliderTransforms(ExplodedTransforms, Transforms);
 
 			// re-parent all the geometry nodes under existing merge node
 			GeometryCollectionAlgo::ParentTransforms(GeometryCollection, MergeNode, SourceElements);
 
 			// update source levels and transforms in our custom attributes
-			for (int32 Element : SourceElements)
-			{
-				ExplodedTransforms[Element] = Transforms[Element];
-				ExplodedVectors[Element] = Transforms[Element].GetLocation();
-			}
+			//for (int32 Element : SourceElements)
+			//{
+			//	ExplodedTransforms[Element] = Transforms[Element];
+			//	ExplodedVectors[Element] = Transforms[Element].GetLocation();
+			//}
 
 			UpdateHierarchyLevelOfChildren(GeometryCollection, MergeNode);
 
diff --git a/Engine/Source/Runtime/Experimental/Chaos/Public/GeometryCollection/GeometryCollectionClusteringUtility.h b/Engine/Source/Runtime/Experimental/Chaos/Public/GeometryCollection/GeometryCollectionClusteringUtility.h
--- a/Engine/Source/Runtime/Experimental/Chaos/Public/GeometryCollection/GeometryCollectionClusteringUtility.h
+++ b/Engine/Source/Runtime/Experimental/Chaos/Public/GeometryCollection/GeometryCollectionClusteringUtility.h
@@ -89,7 +89,7 @@ public:
 	static void CollapseSelectedHierarchy(int8 Level, const TArray<int32>& SelectedBones, FGeometryCollection* GeometryCollection);
 
 	/** reparent source bones under root bone */
-	static void ClusterBonesUnderExistingRoot(FGeometryCollection* GeometryCollection, TArray<int32>& SourceElements);
+	static void ClusterBonesUnderExistingRoot(FGeometryCollection* GeometryCollection, const TArray<int32>& SourceElements);
 
 	/** moved the selected bones closer to the root */
 	static void CollapseHierarchyOneLevel(FGeometryCollection* GeometryCollection, TArray<int32>& SourceElements);
