CreateAssetFromObject -- may now receive an initial path to begin navigation. FractureMode has been extended to take advantage of this new feature by storing the previous path and using it next time the dialog is used.

diff --git a/Engine/Plugins/Experimental/ChaosEditor/Source/FractureEditor/Private/FractureToolGenerators.cpp b/Engine/Plugins/Experimental/ChaosEditor/Source/FractureEditor/Private/FractureToolGenerators.cpp
--- a/Engine/Plugins/Experimental/ChaosEditor/Source/FractureEditor/Private/FractureToolGenerators.cpp
+++ b/Engine/Plugins/Experimental/ChaosEditor/Source/FractureEditor/Private/FractureToolGenerators.cpp
@@ -81,6 +81,7 @@ void UFractureToolGenerateAsset::OpenGenerateAssetDialog(TArray<AActor*>& Actors
 		.AssetFilenameSuffix(TEXT("GeometryCollection"))
 		.HeadingText(LOCTEXT("CreateGeometryCollection_Heading", "Geometry Collection Name"))
 		.CreateButtonText(LOCTEXT("CreateGeometryCollection_ButtonLabel", "Create Geometry Collection"))
+		.AssetPath(AssetPath)
 		.OnCreateAssetAction(FOnPathChosen::CreateUObject(this, &UFractureToolGenerateAsset::OnGenerateAssetPathChosen, Actors))
 	);
 
@@ -99,7 +100,14 @@ void UFractureToolGenerateAsset::OpenGenerateAssetDialog(TArray<AActor*>& Actors
 void UFractureToolGenerateAsset::OnGenerateAssetPathChosen(const FString& InAssetPath, TArray<AActor*> Actors)
 {
 		
-	UGeometryCollectionComponent* GeometryCollectionComponent = nullptr;//  = Cast<UGeometryCollectionComponent>(FractureContext.OriginalPrimitiveComponent);
+	//Record the path
+	int32 LastSlash = INDEX_NONE;
+	if (InAssetPath.FindLastChar('/', LastSlash))
+	{
+		AssetPath = InAssetPath.LeftChop(LastSlash);
+	}
+
+	UGeometryCollectionComponent* GeometryCollectionComponent = nullptr;
 
 	if (Actors.Num() > 0)
 	{
diff --git a/Engine/Plugins/Experimental/ChaosEditor/Source/FractureEditor/Private/FractureToolGenerators.h b/Engine/Plugins/Experimental/ChaosEditor/Source/FractureEditor/Private/FractureToolGenerators.h
--- a/Engine/Plugins/Experimental/ChaosEditor/Source/FractureEditor/Private/FractureToolGenerators.h
+++ b/Engine/Plugins/Experimental/ChaosEditor/Source/FractureEditor/Private/FractureToolGenerators.h
@@ -17,6 +17,7 @@ public:
 
 	UFractureToolGenerateAsset(const FObjectInitializer& ObjInit) 
 		: Super(ObjInit)
+		, AssetPath(FString(TEXT("/Game")))
 	{}
 
 	// UFractureActionTool Interface
@@ -38,6 +39,8 @@ private:
 private:
 	TWeakPtr<FFractureEditorModeToolkit> Toolkit;
 
+	UPROPERTY()
+	FString AssetPath;
 };
 
 
diff --git a/Engine/Source/Editor/UnrealEd/Private/SCreateAssetFromObject.cpp b/Engine/Source/Editor/UnrealEd/Private/SCreateAssetFromObject.cpp
--- a/Engine/Source/Editor/UnrealEd/Private/SCreateAssetFromObject.cpp
+++ b/Engine/Source/Editor/UnrealEd/Private/SCreateAssetFromObject.cpp
@@ -28,8 +28,17 @@ void SCreateAssetFromObject::Construct(const FArguments& InArgs, TSharedPtr<SWin
 	CreateButtonText = InArgs._CreateButtonText;
 	OnCreateAssetAction = InArgs._OnCreateAssetAction;
 
+	if (InArgs._AssetPath.IsEmpty())
+	{
+		AssetPath = FString("/Game");
+	}
+	else
+	{
+		AssetPath = InArgs._AssetPath;
+	}
+
 	bIsReportingError = false;
-	AssetPath = FString("/Game");
+
 	FPathPickerConfig PathPickerConfig;
 	PathPickerConfig.DefaultPath = AssetPath;
 	PathPickerConfig.OnPathSelected = FOnPathSelected::CreateRaw(this, &SCreateAssetFromObject::OnSelectAssetPath);
diff --git a/Engine/Source/Editor/UnrealEd/Public/SCreateAssetFromObject.h b/Engine/Source/Editor/UnrealEd/Public/SCreateAssetFromObject.h
--- a/Engine/Source/Editor/UnrealEd/Public/SCreateAssetFromObject.h
+++ b/Engine/Source/Editor/UnrealEd/Public/SCreateAssetFromObject.h
@@ -33,6 +33,9 @@ public:
 	/** The label for the create button */
 	SLATE_ARGUMENT(FText, CreateButtonText)
 
+	/** The path we begin naviagation at */
+	SLATE_ARGUMENT(FString, AssetPath)
+
 	SLATE_ARGUMENT(FText, DefaultNameOverride)
 
 	/** Action to perform when create clicked */
