Fracture Editor Mode -- Changed cluster behaviour when selected nodes are from different branches.

diff --git a/Engine/Plugins/Experimental/ChaosEditor/Source/FractureEditor/Private/FractureToolClustering.cpp b/Engine/Plugins/Experimental/ChaosEditor/Source/FractureEditor/Private/FractureToolClustering.cpp
--- a/Engine/Plugins/Experimental/ChaosEditor/Source/FractureEditor/Private/FractureToolClustering.cpp
+++ b/Engine/Plugins/Experimental/ChaosEditor/Source/FractureEditor/Private/FractureToolClustering.cpp
@@ -104,38 +104,23 @@ void UFractureToolCluster::Execute(TWeakPtr<FFractureEditorModeToolkit> InToolki
 
 		for (FFractureToolContext& Context : Contexts)
 		{
-			const TManagedArray<TSet<int32>>& Children = Context.GetGeometryCollection()->GetAttribute<TSet<int32>>("Children", FGeometryCollection::TransformGroup);
-			const TManagedArray<int32>& Levels = Context.GetGeometryCollection()->GetAttribute<int32>("Level", FGeometryCollection::TransformGroup);
-			
+			Context.RemoveRootNodes();
 			Context.Sanitize();
-			TMap<int32, TArray<int32>> ClusteredSelection = Context.GetClusteredSelections();
 
-			for (TPair<int32, TArray<int32>>& Group : ClusteredSelection)
-			{
-				const TArray<int32>& CurrentSelection = Group.Value;
+			// Cluster selected bones beneath common parent
+			int32 LowestCommonAncestor = FGeometryCollectionClusteringUtility::FindLowestCommonAncestor(Context.GetGeometryCollection().Get(), Context.GetSelection());
 
-				// sort the selection list so ClusterBonesUnderNewNode() happens in the correct order for leaf nodes
-				TArray<int32> SortedSelectedBones;
-				SortedSelectedBones.Reserve(CurrentSelection.Num());
-				for (int32 SelectedBone : CurrentSelection)
+			if (LowestCommonAncestor != INDEX_NONE)
+			{
+				// ClusterBonesUnderNewNode expects a sibling of the new cluster so we require a child node of the common ancestor.
+				const TManagedArray<TSet<int32>>& Children = Context.GetGeometryCollection()->GetAttribute<TSet<int32>>("Children", FGeometryCollection::TransformGroup);
+				TArray<int32> Siblings = Children[LowestCommonAncestor].Array();
+				if (Siblings.Num() > 0)
 				{
-					if (Children[SelectedBone].Num() > 0)
-					{
-						SortedSelectedBones.Insert(SelectedBone, 0);
-					}
-					else
-					{
-						SortedSelectedBones.Add(SelectedBone);
-					}
+					FGeometryCollectionClusteringUtility::ClusterBonesUnderNewNode(Context.GetGeometryCollection().Get(), Siblings[0], Context.GetSelection(), false);
 				}
-
-				// cluster Selected Bones under the first selected bone
-				int32 InsertAtIndex = SortedSelectedBones[0];
-				CurrentLevelView = Levels[InsertAtIndex];
-
-				FGeometryCollectionClusteringUtility::ClusterBonesUnderNewNode(Context.GetGeometryCollection().Get(), InsertAtIndex, SortedSelectedBones, false);
 			}
-
+			
 			Refresh(Context, Toolkit);
 		}
 
diff --git a/Engine/Source/Runtime/Experimental/Chaos/Private/GeometryCollection/GeometryCollectionClusteringUtility.cpp b/Engine/Source/Runtime/Experimental/Chaos/Private/GeometryCollection/GeometryCollectionClusteringUtility.cpp
--- a/Engine/Source/Runtime/Experimental/Chaos/Private/GeometryCollection/GeometryCollectionClusteringUtility.cpp
+++ b/Engine/Source/Runtime/Experimental/Chaos/Private/GeometryCollection/GeometryCollectionClusteringUtility.cpp
@@ -798,3 +798,52 @@ void FGeometryCollectionClusteringUtility::MoveUpOneHierarchyLevel(FGeometryColl
 }
 
 
+int32 FGeometryCollectionClusteringUtility::FindLowestCommonAncestor(FGeometryCollection* GeometryCollection, const TArray<int32>& SelectedBones)
+{
+	const int32 SelectionCount = SelectedBones.Num();
+	if (SelectionCount == 0)
+	{
+		return INDEX_NONE;
+	}
+
+	int32 LCA = SelectedBones[0];
+	for (int32 Index = 1; Index < SelectionCount; ++Index)
+	{
+		if (LCA == INDEX_NONE)
+		{
+			return INDEX_NONE;
+		}
+		LCA = FindLowestCommonAncestor(GeometryCollection, LCA, SelectedBones[Index]);
+	}
+	return LCA;
+}
+
+int32 FGeometryCollectionClusteringUtility::FindLowestCommonAncestor(FGeometryCollection* GeometryCollection, int32 N0, int32 N1)
+{
+	const TManagedArray<int32>& Parent = GeometryCollection->GetAttribute<int32>("Parent", FGeometryCollection::TransformGroup);
+
+	// Record the path to root from the first 
+	TArray<int32> PathToRoot0;
+	PathToRoot0.Add(N0);
+	while (PathToRoot0.Last() != INDEX_NONE)
+	{
+		PathToRoot0.Add(Parent[PathToRoot0.Last()]);
+	}
+
+	// Traverse from the second node to root and return the first node found that is in the first path.
+	int32 LCA = N1;
+	while (LCA != INDEX_NONE)
+	{
+		if (PathToRoot0.Contains(LCA))
+		{
+			return LCA;
+		}
+		LCA = Parent[LCA];
+	}
+
+	// No common ancestor
+	return INDEX_NONE;
+}
+
+
+
diff --git a/Engine/Source/Runtime/Experimental/Chaos/Public/GeometryCollection/GeometryCollectionClusteringUtility.h b/Engine/Source/Runtime/Experimental/Chaos/Public/GeometryCollection/GeometryCollectionClusteringUtility.h
--- a/Engine/Source/Runtime/Experimental/Chaos/Public/GeometryCollection/GeometryCollectionClusteringUtility.h
+++ b/Engine/Source/Runtime/Experimental/Chaos/Public/GeometryCollection/GeometryCollectionClusteringUtility.h
@@ -109,6 +109,9 @@ public:
 	/** move the selected node up a level in direction of root */
 	static void MoveUpOneHierarchyLevel(FGeometryCollection* GeometryCollection, const TArray<int32>& SelectedBones);
 
+	/** Find the lowest common ancestor index of currently selected nodes. Returns INDEX_NODE if there is no common ancestor. */
+	static int32 FindLowestCommonAncestor(FGeometryCollection* GeometryCollection, const TArray<int32>& SelectedBones);
+
 	static void ValidateResults(FGeometryCollection* GeometryCollection);
 private:
 	static int32 PickBestNodeToMergeTo(const FGeometryCollection* GeometryCollection, const TArray<int32>& SourceElements);
@@ -117,5 +120,5 @@ private:
 	
 	static void RecursivelyUpdateHierarchyLevelOfChildren(TManagedArray<int32>& Levels, const TManagedArray<TSet<int32>>& Children, int32 ParentElement);
 
-
+	static int32 FindLowestCommonAncestor(FGeometryCollection* GeometryCollection, int32 N0, int32 N1);
 };
