Geometry collection assets created from one or more static meshes will have Nanite automatically enabled now, as long as at least one of the input meshes had Nanite enabled originally.

diff --git a/Engine/Plugins/Experimental/ChaosEditor/Source/FractureEditor/Private/FractureToolGenerators.cpp b/Engine/Plugins/Experimental/ChaosEditor/Source/FractureEditor/Private/FractureToolGenerators.cpp
--- a/Engine/Plugins/Experimental/ChaosEditor/Source/FractureEditor/Private/FractureToolGenerators.cpp
+++ b/Engine/Plugins/Experimental/ChaosEditor/Source/FractureEditor/Private/FractureToolGenerators.cpp
@@ -177,6 +177,12 @@ AGeometryCollectionActor* UFractureToolGenerateAsset::ConvertStaticMeshToGeometr
 			if (StaticMeshComponent != nullptr)
 			{
 				UStaticMesh* ComponentStaticMesh = StaticMeshComponent->GetStaticMesh();
+				if (ComponentStaticMesh != nullptr)
+				{
+					// If any of the static meshes have Nanite enabled, also enable on the new geometry collection asset for convenience.
+					FracturedGeometryCollection->EnableNanite |= ComponentStaticMesh->NaniteSettings.bEnabled;
+				}
+
 				FTransform ComponentTransform(StaticMeshComponent->GetComponentTransform());
 				ComponentTransform.SetTranslation((ComponentTransform.GetTranslation() - ActorTransform.GetTranslation()) + ActorOffset);
 
@@ -191,8 +197,15 @@ AGeometryCollectionActor* UFractureToolGenerateAsset::ConvertStaticMeshToGeometr
 
 		FracturedGeometryCollection->InitializeMaterials();
 	}
+
 	AddSingleRootNodeIfRequired(FracturedGeometryCollection);
 
+	if (FracturedGeometryCollection->EnableNanite)
+	{
+		FracturedGeometryCollection->InvalidateCollection();
+		FracturedGeometryCollection->EnsureDataIsCooked(false /* init resources */);
+	}
+
 	return NewActor;
 }
 
diff --git a/Engine/Source/Runtime/Experimental/GeometryCollectionEngine/Private/GeometryCollection/GeometryCollectionObject.cpp b/Engine/Source/Runtime/Experimental/GeometryCollectionEngine/Private/GeometryCollection/GeometryCollectionObject.cpp
--- a/Engine/Source/Runtime/Experimental/GeometryCollectionEngine/Private/GeometryCollection/GeometryCollectionObject.cpp
+++ b/Engine/Source/Runtime/Experimental/GeometryCollectionEngine/Private/GeometryCollection/GeometryCollectionObject.cpp
@@ -736,13 +736,13 @@ bool UGeometryCollection::Modify(bool bAlwaysMarkDirty /*= true*/)
 	return bSuperResult;
 }
 
-void UGeometryCollection::EnsureDataIsCooked()
+void UGeometryCollection::EnsureDataIsCooked(bool bInitResources)
 {
 	if (StateGuid != LastBuiltGuid)
 	{
 		CreateSimulationDataImp(/*bCopyFromDDC=*/ true);
 		
-		if (FApp::CanEverRender())
+		if (FApp::CanEverRender() && bInitResources)
 		{
 			NaniteData->InitResources(this);
 		}
diff --git a/Engine/Source/Runtime/Experimental/GeometryCollectionEngine/Public/GeometryCollection/GeometryCollectionObject.h b/Engine/Source/Runtime/Experimental/GeometryCollectionEngine/Public/GeometryCollection/GeometryCollectionObject.h
--- a/Engine/Source/Runtime/Experimental/GeometryCollectionEngine/Public/GeometryCollection/GeometryCollectionObject.h
+++ b/Engine/Source/Runtime/Experimental/GeometryCollectionEngine/Public/GeometryCollection/GeometryCollectionObject.h
@@ -152,7 +152,7 @@ public:
 	void Serialize(FArchive& Ar);
 
 #if WITH_EDITOR
-	void EnsureDataIsCooked();
+	void EnsureDataIsCooked(bool bInitResources = true);
 #endif
 
 	/** Accessors for internal geometry collection */
