Fracture Mode -- Cleared selection and bone hierarchy level following geometry collection reset to ensure valid behavior.

diff --git a/Engine/Plugins/Experimental/ChaosEditor/Source/FractureEditor/Private/FractureToolGenerators.cpp b/Engine/Plugins/Experimental/ChaosEditor/Source/FractureEditor/Private/FractureToolGenerators.cpp
--- a/Engine/Plugins/Experimental/ChaosEditor/Source/FractureEditor/Private/FractureToolGenerators.cpp
+++ b/Engine/Plugins/Experimental/ChaosEditor/Source/FractureEditor/Private/FractureToolGenerators.cpp
@@ -136,6 +136,8 @@ void UFractureToolGenerateAsset::OnGenerateAssetPathChosen(const FString& InAsse
 			TSharedPtr<FFractureEditorModeToolkit> SharedToolkit(Toolkit.Pin());
 			SharedToolkit->SetOutlinerComponents({ GeometryCollectionComponent });
 			SharedToolkit->SetBoneSelection(GeometryCollectionComponent, EditBoneColor.GetSelectedBones(), true);
+
+			SharedToolkit->OnSetLevelViewValue(-1);
 		}
 		
 		GeometryCollectionComponent->MarkRenderDynamicDataDirty();
@@ -363,7 +365,12 @@ void UFractureToolResetAsset::Execute(TWeakPtr<FFractureEditorModeToolkit> InToo
 			}
 			GeometryCollectionObject->MarkPackageDirty();
 		}
+		
+		FScopedColorEdit EditBoneColor = GeometryCollectionComponent->EditBoneSelection();
+		EditBoneColor.ResetBoneSelection();
+		EditBoneColor.ResetHighlightedBones();
 	}
+	InToolkit.Pin()->OnSetLevelViewValue(-1);
 	InToolkit.Pin()->SetOutlinerComponents(GeomCompSelection.Array());
 }
 
