Fix reference actor not being used in fracture plane cut

#rnx

diff --git a/Engine/Plugins/Experimental/ChaosEditor/Source/FractureEditor/Private/FractureToolPlaneCut.cpp b/Engine/Plugins/Experimental/ChaosEditor/Source/FractureEditor/Private/FractureToolPlaneCut.cpp
--- a/Engine/Plugins/Experimental/ChaosEditor/Source/FractureEditor/Private/FractureToolPlaneCut.cpp
+++ b/Engine/Plugins/Experimental/ChaosEditor/Source/FractureEditor/Private/FractureToolPlaneCut.cpp
@@ -41,7 +41,7 @@ void UFractureToolPlaneCut::RegisterUICommand( FFractureEditorCommands* BindingC
 
 void UFractureToolPlaneCut::Render(const FSceneView* View, FViewport* Viewport, FPrimitiveDrawInterface* PDI)
 {
-	const UFracturePlaneCutSettings* LocalCutSettings = GetMutableDefault<UFracturePlaneCutSettings>();
+	const UFracturePlaneCutSettings* LocalCutSettings = PlaneCutSettings;
 
 	if (LocalCutSettings->ReferenceActor != nullptr) // so we update with the ref actor realtime
 	{
@@ -101,7 +101,7 @@ int32 UFractureToolPlaneCut::ExecuteFracture(const FFractureToolContext& Fractur
 {
 	if (FractureContext.IsValid())
 	{
-		const UFracturePlaneCutSettings* LocalCutSettings = GetMutableDefault<UFracturePlaneCutSettings>();
+		const UFracturePlaneCutSettings* LocalCutSettings = PlaneCutSettings;
 
 		TArray<FPlane> CuttingPlanes;
 		if (LocalCutSettings->ReferenceActor != nullptr)
