FractureMode -- Flattening tool cleans up dangling cluster nodes left over after operation is complete.

diff --git a/Engine/Plugins/Experimental/ChaosEditor/Source/FractureEditor/Private/FractureToolClustering.cpp b/Engine/Plugins/Experimental/ChaosEditor/Source/FractureEditor/Private/FractureToolClustering.cpp
--- a/Engine/Plugins/Experimental/ChaosEditor/Source/FractureEditor/Private/FractureToolClustering.cpp
+++ b/Engine/Plugins/Experimental/ChaosEditor/Source/FractureEditor/Private/FractureToolClustering.cpp
@@ -53,6 +53,10 @@ void UFractureToolFlattenAll::Execute(TWeakPtr<FFractureEditorModeToolkit> InToo
 				TArray<int32> LeafBones;
 				FGeometryCollectionClusteringUtility::GetLeafBones(Context.GetGeometryCollection().Get(), ClusterIndex, LeafBones);
 				FGeometryCollectionClusteringUtility::ClusterBonesUnderExistingNode(Context.GetGeometryCollection().Get(), ClusterIndex, LeafBones);
+
+				// Cleanup: Remove any clusters remaining in the flattened branch.
+				FGeometryCollectionClusteringUtility::RemoveDanglingClusters(Context.GetGeometryCollection().Get());
+
 				CurrentLevelView = Levels[ClusterIndex] + 1;
 			}
 
@@ -107,21 +111,24 @@ void UFractureToolCluster::Execute(TWeakPtr<FFractureEditorModeToolkit> InToolki
 			Context.RemoveRootNodes();
 			Context.Sanitize();
 
-			// Cluster selected bones beneath common parent
-			int32 LowestCommonAncestor = FGeometryCollectionClusteringUtility::FindLowestCommonAncestor(Context.GetGeometryCollection().Get(), Context.GetSelection());
-
-			if (LowestCommonAncestor != INDEX_NONE)
+			if (Context.GetSelection().Num() > 1)
 			{
-				// ClusterBonesUnderNewNode expects a sibling of the new cluster so we require a child node of the common ancestor.
-				const TManagedArray<TSet<int32>>& Children = Context.GetGeometryCollection()->GetAttribute<TSet<int32>>("Children", FGeometryCollection::TransformGroup);
-				TArray<int32> Siblings = Children[LowestCommonAncestor].Array();
-				if (Siblings.Num() > 0)
+				// Cluster selected bones beneath common parent
+				int32 LowestCommonAncestor = FGeometryCollectionClusteringUtility::FindLowestCommonAncestor(Context.GetGeometryCollection().Get(), Context.GetSelection());
+
+				if (LowestCommonAncestor != INDEX_NONE)
 				{
-					FGeometryCollectionClusteringUtility::ClusterBonesUnderNewNode(Context.GetGeometryCollection().Get(), Siblings[0], Context.GetSelection(), false);
+					// ClusterBonesUnderNewNode expects a sibling of the new cluster so we require a child node of the common ancestor.
+					const TManagedArray<TSet<int32>>& Children = Context.GetGeometryCollection()->GetAttribute<TSet<int32>>("Children", FGeometryCollection::TransformGroup);
+					TArray<int32> Siblings = Children[LowestCommonAncestor].Array();
+					if (Siblings.Num() > 0)
+					{
+						FGeometryCollectionClusteringUtility::ClusterBonesUnderNewNode(Context.GetGeometryCollection().Get(), Siblings[0], Context.GetSelection(), true);
+					}
 				}
+
+				Refresh(Context, Toolkit);
 			}
-			
-			Refresh(Context, Toolkit);
 		}
 
 		if (CurrentLevelView != Toolkit->GetLevelViewValue())
diff --git a/Engine/Source/Runtime/Experimental/Chaos/Private/GeometryCollection/GeometryCollectionClusteringUtility.cpp b/Engine/Source/Runtime/Experimental/Chaos/Private/GeometryCollection/GeometryCollectionClusteringUtility.cpp
--- a/Engine/Source/Runtime/Experimental/Chaos/Private/GeometryCollection/GeometryCollectionClusteringUtility.cpp
+++ b/Engine/Source/Runtime/Experimental/Chaos/Private/GeometryCollection/GeometryCollectionClusteringUtility.cpp
@@ -15,6 +15,7 @@ void FGeometryCollectionClusteringUtility::ClusterBonesUnderNewNode(FGeometryCol
 	TManagedArray<FString>& BoneNames = GeometryCollection->BoneName;
 	TManagedArray<int32>& Parents = GeometryCollection->Parent;
 	TManagedArray<TSet<int32>>& Children = GeometryCollection->Children;
+	TManagedArray<int32>& SimType = GeometryCollection->SimulationType;
 
 	// insert a new node between the selected bones and their shared parent
 	int NewBoneIndex = GeometryCollection->AddElements(1, FGeometryCollection::TransformGroup);
@@ -25,6 +26,7 @@ void FGeometryCollectionClusteringUtility::ClusterBonesUnderNewNode(FGeometryCol
 	BoneNames[NewBoneIndex] = BoneNames[SourceBoneIndex];
 	Parents[NewBoneIndex] = OriginalParentIndex;
 	Children[NewBoneIndex] = TSet<int32>(SelectedBones);
+	SimType[NewBoneIndex] = FGeometryCollection::ESimulationTypes::FST_Clustered;
 
 	Transforms[NewBoneIndex] = FTransform::Identity;
 
@@ -845,5 +847,32 @@ int32 FGeometryCollectionClusteringUtility::FindLowestCommonAncestor(FGeometryCo
 	return INDEX_NONE;
 }
 
+void FGeometryCollectionClusteringUtility::RemoveDanglingClusters(FGeometryCollection* GeometryCollection)
+{
+	check(GeometryCollection);
+
+	const TManagedArray<TSet<int32>>& Children = GeometryCollection->Children;
+	const TManagedArray<FTransform>& Transforms = GeometryCollection->Transform;
+	const TManagedArray<int32>& SimulationType = GeometryCollection->SimulationType;
+
+	const int32 TransformCount = Transforms.Num();
+	TArray<int32> DeletionList;
+	for (int32 Idx = 0; Idx < TransformCount; ++Idx)
+	{
+		if(GeometryCollection->IsClustered(Idx))
+		{
+			if (Children[Idx].Num() == 0)
+			{
+				DeletionList.Add(Idx);
+			}
+		}
+	}
+
+	if (DeletionList.Num())
+	{
+		GeometryCollection->RemoveElements(FGeometryCollection::TransformGroup, DeletionList);
+	}
+}
+
 
 
diff --git a/Engine/Source/Runtime/Experimental/Chaos/Public/GeometryCollection/GeometryCollectionClusteringUtility.h b/Engine/Source/Runtime/Experimental/Chaos/Public/GeometryCollection/GeometryCollectionClusteringUtility.h
--- a/Engine/Source/Runtime/Experimental/Chaos/Public/GeometryCollection/GeometryCollectionClusteringUtility.h
+++ b/Engine/Source/Runtime/Experimental/Chaos/Public/GeometryCollection/GeometryCollectionClusteringUtility.h
@@ -112,6 +112,9 @@ public:
 	/** Find the lowest common ancestor index of currently selected nodes. Returns INDEX_NODE if there is no common ancestor. */
 	static int32 FindLowestCommonAncestor(FGeometryCollection* GeometryCollection, const TArray<int32>& SelectedBones);
 
+	/** Delete any cluster nodes discovered to have no children */
+	static void RemoveDanglingClusters(FGeometryCollection* GeometryCollection);
+
 	static void ValidateResults(FGeometryCollection* GeometryCollection);
 private:
 	static int32 PickBestNodeToMergeTo(const FGeometryCollection* GeometryCollection, const TArray<int32>& SourceElements);
