Fracture Mode -- new icon for Reset tool

diff --git a/Engine/Plugins/Experimental/ChaosEditor/Source/FractureEditor/Private/FractureEditorStyle.cpp b/Engine/Plugins/Experimental/ChaosEditor/Source/FractureEditor/Private/FractureEditorStyle.cpp
--- a/Engine/Plugins/Experimental/ChaosEditor/Source/FractureEditor/Private/FractureEditorStyle.cpp
+++ b/Engine/Plugins/Experimental/ChaosEditor/Source/FractureEditor/Private/FractureEditorStyle.cpp
@@ -5,6 +5,7 @@
 #include "Styling/SlateTypes.h"
 #include "Styling/CoreStyle.h"
 #include "EditorStyleSet.h"
+#include "Styling/SlateStyleMacros.h"
 
 FName FFractureEditorStyle::StyleName("FractureEditorStyle");
 
@@ -92,10 +93,17 @@ FFractureEditorStyle::FFractureEditorStyle()
 
 	Set("FractureEditor.GenerateAsset", new FSlateImageBrush(RootToContentDir(TEXT("GenerateAsset_48x.png")), IconSize));
 	Set("FractureEditor.GenerateAsset.Small", new FSlateImageBrush(RootToContentDir(TEXT("GenerateAsset_48x.png")), SmallIconSize));
-
-
-	Set("LevelEditor.FractureMode", new FSlateImageBrush(RootToContentDir(TEXT("FractureMode.png")), FVector2D(40.0f, 40.0f)));
-	Set("LevelEditor.FractureMode.Small", new FSlateImageBrush(RootToContentDir(TEXT("FractureMode.png")), FVector2D(20.0f, 20.0f)));
+	Set("FractureEditor.ResetAsset", new IMAGE_BRUSH_SVG("Starship/FractureReset", FVector2D(20.0f, 20.0f)));
+
+	if (FCoreStyle::IsStarshipStyle())
+	{
+		Set("LevelEditor.FractureMode", new IMAGE_BRUSH_SVG("Starship/fracture", FVector2D(20.0f, 20.0f)));
+	}
+	else
+	{
+		Set("LevelEditor.FractureMode", new FSlateImageBrush(RootToContentDir(TEXT("FractureMode.png")), FVector2D(40.0f, 40.0f)));
+		Set("LevelEditor.FractureMode.Small", new FSlateImageBrush(RootToContentDir(TEXT("FractureMode.png")), FVector2D(20.0f, 20.0f)));
+	}
 
 	FSlateStyleRegistry::RegisterSlateStyle(*this);
 }
