GeometryCollection -- We may now parent bones to rigids as well as clusters. Bones may be flagged as non-simulatable and the solver respects this: no mass properties, simplicials, etc are produced and a particle is not injected into the solver.

Fracture Mode -- A prototype tool for handplacing embedded geometry is implemented.

diff --git a/Engine/Plugins/Experimental/ChaosEditor/Source/FractureEditor/Private/FractureEditorCommands.h b/Engine/Plugins/Experimental/ChaosEditor/Source/FractureEditor/Private/FractureEditorCommands.h
--- a/Engine/Plugins/Experimental/ChaosEditor/Source/FractureEditor/Private/FractureEditorCommands.h
+++ b/Engine/Plugins/Experimental/ChaosEditor/Source/FractureEditor/Private/FractureEditorCommands.h
@@ -38,6 +38,7 @@ class FFractureEditorCommands : public TCommands<FFractureEditorCommands>
 		// Generate Commands
 		TSharedPtr< FUICommandInfo > GenerateAsset;
 		TSharedPtr< FUICommandInfo > ResetAsset;
+		TSharedPtr< FUICommandInfo > AddEmbeddedGeometry;
 		
 		// Fracture Commands
 		TSharedPtr< FUICommandInfo > Uniform;
diff --git a/Engine/Plugins/Experimental/ChaosEditor/Source/FractureEditor/Private/FractureEditorModeToolkit.cpp b/Engine/Plugins/Experimental/ChaosEditor/Source/FractureEditor/Private/FractureEditorModeToolkit.cpp
--- a/Engine/Plugins/Experimental/ChaosEditor/Source/FractureEditor/Private/FractureEditorModeToolkit.cpp
+++ b/Engine/Plugins/Experimental/ChaosEditor/Source/FractureEditor/Private/FractureEditorModeToolkit.cpp
@@ -481,6 +481,7 @@ void FFractureEditorModeToolkit::BuildToolPalette(FName PaletteIndex, class FToo
 	{
 		ToolbarBuilder.AddToolBarButton(Commands.GenerateAsset);
 		ToolbarBuilder.AddToolBarButton(Commands.ResetAsset);
+		ToolbarBuilder.AddToolBarButton(Commands.AddEmbeddedGeometry);
 	}
 	else if (PaletteIndex == TEXT("Select"))
 	{
diff --git a/Engine/Plugins/Experimental/ChaosEditor/Source/FractureEditor/Private/FractureToolGenerators.cpp b/Engine/Plugins/Experimental/ChaosEditor/Source/FractureEditor/Private/FractureToolGenerators.cpp
--- a/Engine/Plugins/Experimental/ChaosEditor/Source/FractureEditor/Private/FractureToolGenerators.cpp
+++ b/Engine/Plugins/Experimental/ChaosEditor/Source/FractureEditor/Private/FractureToolGenerators.cpp
@@ -10,12 +10,14 @@
 #include "AssetRegistryModule.h"
 #include "Layers/LayersSubsystem.h"
 #include "Engine/SkeletalMesh.h"
+#include "Engine/StaticMesh.h"
 
 #include "GeometryCollection/GeometryCollectionComponent.h"
 #include "GeometryCollection/GeometryCollectionActor.h"
 #include "GeometryCollection/GeometryCollection.h"
 #include "GeometryCollection/GeometryCollectionClusteringUtility.h"
 #include "GeometryCollection/GeometryCollectionConversion.h"
+#include "FractureToolContext.h"
 
 
 #define LOCTEXT_NAMESPACE "FractureToolGenerators"
@@ -44,8 +46,8 @@ void UFractureToolGenerateAsset::RegisterUICommand(FFractureEditorCommands* Bind
 
 bool UFractureToolGenerateAsset::CanExecute() const
 {
-	// We can execute this command if static meshes are selected.
-	return IsStaticMeshSelected();
+	// We can execute this command if only static meshes are selected.
+	return (IsStaticMeshSelected() && (!IsGeometryCollectionSelected()));
 }
 
 void UFractureToolGenerateAsset::Execute(TWeakPtr<FFractureEditorModeToolkit> InToolkit)
@@ -374,4 +376,134 @@ void UFractureToolResetAsset::Execute(TWeakPtr<FFractureEditorModeToolkit> InToo
 	InToolkit.Pin()->SetOutlinerComponents(GeomCompSelection.Array());
 }
 
+
+
+
+FText UFractureToolAddEmbeddedGeometry::GetDisplayText() const
+{
+	return FText(NSLOCTEXT("Fracture", "FractureToolAddEmbeddedGeometry", "Embed"));
+}
+
+FText UFractureToolAddEmbeddedGeometry::GetTooltipText() const
+{
+	return FText(NSLOCTEXT("Fracture", "FractureToolAddEmbeddedGeometryTooltip", "Embed selected static mesh as passive geometry parented to nearest bone. Will be lost if GeometryCollection is Reset!"));
+}
+
+FSlateIcon UFractureToolAddEmbeddedGeometry::GetToolIcon() const
+{
+	return FSlateIcon("FractureEditorStyle", "FractureEditor.AddEmbeddedGeometry");
+}
+
+void UFractureToolAddEmbeddedGeometry::RegisterUICommand(FFractureEditorCommands* BindingContext)
+{
+	UI_COMMAND_EXT(BindingContext, UICommandInfo, "AddEmbeddedGeometry", "Embed", "Embed", EUserInterfaceActionType::ToggleButton, FInputChord());
+	BindingContext->AddEmbeddedGeometry = UICommandInfo;
+}
+
+bool UFractureToolAddEmbeddedGeometry::CanExecute() const
+{
+	return (IsStaticMeshSelected() && IsGeometryCollectionSelected());
+}
+
+void UFractureToolAddEmbeddedGeometry::Execute(TWeakPtr<FFractureEditorModeToolkit> InToolkit)
+{
+	if (!InToolkit.IsValid())
+	{
+		return;
+	}
+
+	// Get selected static meshes
+	TArray<UStaticMeshComponent*> SelectedStaticMeshComponents = GetSelectedStaticMeshComponents();
+
+	TArray<FFractureToolContext> Contexts = GetFractureToolContexts();
+
+	for (FFractureToolContext& Context : Contexts)
+	{
+		FGeometryCollectionEdit GeometryCollectionEdit = Context.GetGeometryCollectionComponent()->EditRestCollection(GeometryCollection::EEditUpdate::RestPhysicsDynamic);
+		UGeometryCollection* FracturedGeometryCollection = GeometryCollectionEdit.GetRestCollection();
+
+		FFractureToolContext::FGeometryCollectionPtr GeometryCollection = Context.GetGeometryCollection();
+		TManagedArray<int32>& Parent = GeometryCollection->Parent;
+		TManagedArray<int32>& SimType = GeometryCollection->SimulationType;
+		TManagedArray<TSet<int32>>& Children = GeometryCollection->Children;
+		TManagedArray<bool>& CollectionSimulatableParticles = 
+			GeometryCollection->GetAttribute<bool>(
+				FGeometryCollection::SimulatableParticlesAttribute, FTransformCollection::TransformGroup);
+
+		const FTransform TargetActorTransform(Context.GetGeometryCollectionComponent()->GetOwner()->GetTransform());
+		
+		Context.ConvertSelectionToLeafNodes();
+		const TArray<int32>& SelectedBones = Context.GetSelection();
+		for (const int32 SelectedBone : SelectedBones)
+		{
+			for (UStaticMeshComponent* SelectedStaticMeshComponent : SelectedStaticMeshComponents)
+			{
+				const AActor* Actor = SelectedStaticMeshComponent->GetOwner();
+				check(Actor)
+
+				const FTransform SMActorTransform(Actor->GetTransform());
+
+				UStaticMesh* ComponentStaticMesh = SelectedStaticMeshComponent->GetStaticMesh();
+				if (ComponentStaticMesh)
+				{
+					// If any of the static meshes have Nanite enabled, also enable on the new geometry collection asset for convenience.
+					FracturedGeometryCollection->EnableNanite |= ComponentStaticMesh->NaniteSettings.bEnabled;
+				}
+
+				FTransform ComponentTransform = SMActorTransform.GetRelativeTransform(TargetActorTransform);
+				decltype(FGeometryCollectionSource::SourceMaterial) SourceMaterials(SelectedStaticMeshComponent->GetMaterials());
+
+				// Add static mesh geometry to context's geometry collection.
+				FGeometryCollectionConversion::AppendStaticMesh(ComponentStaticMesh, SourceMaterials, ComponentTransform, FracturedGeometryCollection, true);
+
+				int32 NewTransformIndex = Parent.Num() - 1;
+				check(NewTransformIndex > -1)
+
+				// Reparent the new transform to the selected bone
+				Parent[NewTransformIndex] = SelectedBone;
+				Children[SelectedBone].Add(NewTransformIndex);
+
+				// Set properties appropriate to embedded geometry
+				SimType[NewTransformIndex] = FGeometryCollection::ESimulationTypes::FST_Not_Simulatable;
+				CollectionSimulatableParticles[NewTransformIndex] = false;
+			}	
+		}
+
+		FracturedGeometryCollection->InitializeMaterials();
+
+		FGeometryCollectionClusteringUtility::UpdateHierarchyLevelOfChildren(GeometryCollection.Get(), -1);
+		Context.GetGeometryCollectionComponent()->MarkRenderStateDirty();
+		
+		FracturedGeometryCollection->MarkPackageDirty();
+	}
+
+	SetOutlinerComponents(Contexts, InToolkit.Pin().Get());
+}
+
+TArray<UStaticMeshComponent*> UFractureToolAddEmbeddedGeometry::GetSelectedStaticMeshComponents()
+{
+	TArray<UStaticMeshComponent*> SelectedStaticMeshComponents;
+
+	USelection* SelectedActors = GEditor->GetSelectedActors();
+	for (FSelectionIterator Iter(*SelectedActors); Iter; ++Iter)
+	{
+		AActor* Actor = Cast<AActor>(*Iter);
+		if (Actor)
+		{
+			// We don't want to include any static meshes owned by a geometry collection.
+			// (This might be an ISMC or swap out static geometry.)
+			if (Actor->FindComponentByClass<UGeometryCollectionComponent>())
+			{
+				continue;
+			}
+			
+			TArray<UStaticMeshComponent*> StaticMeshComponents;
+			Actor->GetComponents<UStaticMeshComponent>(StaticMeshComponents, true);
+			SelectedStaticMeshComponents.Append(StaticMeshComponents);
+		}
+	}
+
+	return SelectedStaticMeshComponents;
+}
+
 #undef LOCTEXT_NAMESPACE
diff --git a/Engine/Plugins/Experimental/ChaosEditor/Source/FractureEditor/Private/FractureToolGenerators.h b/Engine/Plugins/Experimental/ChaosEditor/Source/FractureEditor/Private/FractureToolGenerators.h
--- a/Engine/Plugins/Experimental/ChaosEditor/Source/FractureEditor/Private/FractureToolGenerators.h
+++ b/Engine/Plugins/Experimental/ChaosEditor/Source/FractureEditor/Private/FractureToolGenerators.h
@@ -61,3 +61,24 @@ public:
 	virtual bool CanExecute() const override;
 };
 
+
+UCLASS()
+class UFractureToolAddEmbeddedGeometry : public UFractureActionTool
+{
+public:
+	GENERATED_BODY()
+
+	UFractureToolAddEmbeddedGeometry(const FObjectInitializer& ObjInit) : Super(ObjInit) {}
+
+	// UFractureActionTool Interface
+	virtual FText GetDisplayText() const override;
+	virtual FText GetTooltipText() const override;
+	virtual FSlateIcon GetToolIcon() const override;
+	virtual void RegisterUICommand(FFractureEditorCommands* BindingContext) override;
+	virtual void Execute(TWeakPtr<FFractureEditorModeToolkit> InToolkit) override;
+	virtual bool CanExecute() const override;
+
+private:
+	static TArray<UStaticMeshComponent*>  GetSelectedStaticMeshComponents();
+};
+
diff --git a/Engine/Source/Runtime/Experimental/Chaos/Private/GeometryCollection/GeometryCollectionClusteringUtility.cpp b/Engine/Source/Runtime/Experimental/Chaos/Private/GeometryCollection/GeometryCollectionClusteringUtility.cpp
--- a/Engine/Source/Runtime/Experimental/Chaos/Private/GeometryCollection/GeometryCollectionClusteringUtility.cpp
+++ b/Engine/Source/Runtime/Experimental/Chaos/Private/GeometryCollection/GeometryCollectionClusteringUtility.cpp
@@ -111,7 +111,7 @@ void FGeometryCollectionClusteringUtility::ClusterAllBonesUnderNewRoot(FGeometry
 	BoneNames[RootNoneIndex] = "ClusterBone";
 	Parents[RootNoneIndex] = FGeometryCollection::Invalid;
 	Children[RootNoneIndex] = TSet<int32>(ChildBones);
-	SimulationType[RootNoneIndex] = FGeometryCollection::ESimulationTypes::FST_Rigid;
+	SimulationType[RootNoneIndex] = FGeometryCollection::ESimulationTypes::FST_Clustered;
 	check(GeometryCollection->IsTransform(RootNoneIndex));
 
 	if (GeometryCollection->HasAttribute("ExplodedVector", FGeometryCollection::TransformGroup) &&
@@ -137,8 +137,6 @@ void FGeometryCollectionClusteringUtility::ClusterAllBonesUnderNewRoot(FGeometry
 	for (int32 ChildBoneIndex : ChildBones)
 	{
 		Parents[ChildBoneIndex] = RootNoneIndex;
-		GeometryCollection->SimulationType[ChildBoneIndex] = FGeometryCollection::ESimulationTypes::FST_Clustered;
-
 	}
 
 	Transforms[RootNoneIndex] = FTransform::Identity;
diff --git a/Engine/Source/Runtime/Experimental/Chaos/Private/GeometryCollection/GeometryCollectionUtility.cpp b/Engine/Source/Runtime/Experimental/Chaos/Private/GeometryCollection/GeometryCollectionUtility.cpp
--- a/Engine/Source/Runtime/Experimental/Chaos/Private/GeometryCollection/GeometryCollectionUtility.cpp
+++ b/Engine/Source/Runtime/Experimental/Chaos/Private/GeometryCollection/GeometryCollectionUtility.cpp
@@ -38,10 +38,12 @@ namespace GeometryCollection
 		TManagedArray<int32>&  MaterialIndex = RestCollection->MaterialIndex;
 		TManagedArray<int32>&  MaterialID = RestCollection->MaterialID;
 		TManagedArray<FTransform>&  Transform = RestCollection->Transform;
+		TManagedArray<int32>& SimType = RestCollection->SimulationType;
 
 		// set the particle information
 		Transform[0] = center;
 		Transform[0].NormalizeRotation();
+		SimType[0] = FGeometryCollection::ESimulationTypes::FST_Rigid;
 
 		// set the vertex information
 		int32 Index = 0;
diff --git a/Engine/Source/Runtime/Experimental/Chaos/Private/PhysicsProxy/GeometryCollectionPhysicsProxy.cpp b/Engine/Source/Runtime/Experimental/Chaos/Private/PhysicsProxy/GeometryCollectionPhysicsProxy.cpp
--- a/Engine/Source/Runtime/Experimental/Chaos/Private/PhysicsProxy/GeometryCollectionPhysicsProxy.cpp
+++ b/Engine/Source/Runtime/Experimental/Chaos/Private/PhysicsProxy/GeometryCollectionPhysicsProxy.cpp
@@ -688,26 +688,26 @@ void FGeometryCollectionPhysicsProxy::InitializeBodiesPT(Chaos::TPBDRigidsSolver
 		TArray<FTransform> Transform;
 		GeometryCollectionAlgo::GlobalMatrices(DynamicCollection.Transform, DynamicCollection.Parent, Transform);
 
-		const int NumRigids = 0; // ryan - Since we're doing SOA, we start at zero?
+		//const int NumRigids = 0; // ryan - Since we're doing SOA, we start at zero?
+		int NumRigids = 0;
 		BaseParticleIndex = NumRigids;
 
 		// Count geometry collection leaf node particles to add
 		int NumSimulatedParticles = 0;
-		int NumLeafNodes = 0;
 		for (int32 Idx = 0; Idx < SimulatableParticles.Num(); ++Idx)
 		{
 			NumSimulatedParticles += SimulatableParticles[Idx];
-			NumLeafNodes += (SimulatableParticles[Idx] && Children[Idx].Num() == 0);
+			NumRigids += (SimulatableParticles[Idx] && !RestCollection->IsClustered(Idx));
 		}
 
 		// Add entries into simulation array
 		RigidsSolver->GetEvolution()->ReserveParticles(NumSimulatedParticles);
-		TArray<Chaos::TPBDGeometryCollectionParticleHandle<float, 3>*> Handles = RigidsSolver->GetEvolution()->CreateGeometryCollectionParticles(NumLeafNodes);
+		TArray<Chaos::TPBDGeometryCollectionParticleHandle<float, 3>*> Handles = RigidsSolver->GetEvolution()->CreateGeometryCollectionParticles(NumRigids);
 
 		int32 NextIdx = 0;
 		for (int32 Idx = 0; Idx < SimulatableParticles.Num(); ++Idx)
 		{
-			if (SimulatableParticles[Idx] && Children[Idx].Num() == 0) // Leaf node
+			if (SimulatableParticles[Idx] && !RestCollection->IsClustered(Idx))
 			{
 				// todo: Unblocked read access of game thread data on the physics thread.
 
@@ -845,9 +845,10 @@ void FGeometryCollectionPhysicsProxy::InitializeBodiesPT(Chaos::TPBDRigidsSolver
 			SubTreeContainsSimulatableParticle.SetNumZeroed(RecursiveOrder.Num());
 			for (const int32 TransformGroupIndex : RecursiveOrder)
 			{
-				if (Children[TransformGroupIndex].Num() == 0)
+				if (SimulatableParticles[TransformGroupIndex] && !RestCollection->IsClustered(TransformGroupIndex))
+
 				{
-					// Leaf node
+					// Rigid node
 					SubTreeContainsSimulatableParticle[TransformGroupIndex] =
 						SolverParticleHandles[TransformGroupIndex] != nullptr;
 				}
@@ -943,8 +944,7 @@ void FGeometryCollectionPhysicsProxy::InitializeBodiesPT(Chaos::TPBDRigidsSolver
 			// will optionally do this, but we switch that functionality off in BuildClusters().
 			for(int32 TransformGroupIndex = 0; TransformGroupIndex < NumTransforms; ++TransformGroupIndex)
 			{
-				const TManagedArray<TSet<int32>>& RestChildren = RestCollection->Children;
-				if (RestChildren[TransformGroupIndex].Num() > 0)
+				if (RestCollection->IsClustered(TransformGroupIndex))
 				{
 					if (SolverClusterHandles[TransformGroupIndex])
 					{
@@ -1772,6 +1772,7 @@ void FGeometryCollectionPhysicsProxy::InitializeSharedCollisionStructures(
 	const TManagedArray<int32>& BoneMap = RestCollection.BoneMap;
 	const TManagedArray<int32>& Parent = RestCollection.Parent;
 	const TManagedArray<TSet<int32>>& Children = RestCollection.Children;
+	const TManagedArray<int32>& SimulationType = RestCollection.SimulationType;
 	TManagedArray<bool>& CollectionSimulatableParticles =
 		RestCollection.GetAttribute<bool>(
 			FGeometryCollection::SimulatableParticlesAttribute, FTransformCollection::TransformGroup);
@@ -1853,7 +1854,7 @@ void FGeometryCollectionPhysicsProxy::InitializeSharedCollisionStructures(
 	for (int32 GeometryIndex = 0; GeometryIndex < NumGeometries; GeometryIndex++)
 	{
 		const int32 TransformGroupIndex = TransformIndex[GeometryIndex];
-		if (CollectionSimulatableParticles[TransformGroupIndex])
+		if (SimulationType[TransformGroupIndex]!= FGeometryCollection::ESimulationTypes::FST_Not_Simulatable)
 		{
 			TUniquePtr<TTriangleMesh<float>> TriMesh(
 				CreateTriangleMesh(
@@ -1958,6 +1959,10 @@ void FGeometryCollectionPhysicsProxy::InitializeSharedCollisionStructures(
 			TotalVolume += MassProperties.Volume;
 			TriangleMeshesArray[TransformGroupIndex] = MoveTemp(TriMesh);
 		}
+		else
+		{
+			CollectionSimulatableParticles[TransformGroupIndex] = false;
+		}
 	}
 
 	//User provides us with total mass or density.
@@ -2196,7 +2201,7 @@ void FGeometryCollectionPhysicsProxy::InitializeSharedCollisionStructures(
 		//build collision structures depth first
 		for (const int32 TransformGroupIndex : RecursiveOrder)
 		{
-			if (Children[TransformGroupIndex].Num())	//only care about clusters at this point
+			if (RestCollection.IsClustered(TransformGroupIndex))
 			{
 				const int32 ClusterTransformIdx = TransformGroupIndex;
 				//update mass 
diff --git a/Engine/Source/Runtime/Experimental/Chaos/Public/GeometryCollection/GeometryCollection.h b/Engine/Source/Runtime/Experimental/Chaos/Public/GeometryCollection/GeometryCollection.h
--- a/Engine/Source/Runtime/Experimental/Chaos/Public/GeometryCollection/GeometryCollection.h
+++ b/Engine/Source/Runtime/Experimental/Chaos/Public/GeometryCollection/GeometryCollection.h
@@ -199,6 +199,7 @@ public:
 
 	FORCEINLINE bool IsGeometry(int32 Element) const { return TransformToGeometryIndex[Element] != INDEX_NONE; }
 	FORCEINLINE bool IsClustered(int32 Element) const { const TManagedArray<int32>& SimType = SimulationType;  return !!(SimType[Element] == ESimulationTypes::FST_Clustered); }
+	FORCEINLINE bool IsRigid(int32 Element) const { const TManagedArray<int32>& SimType = SimulationType;  return !!(SimType[Element] == ESimulationTypes::FST_Rigid); }
 	FORCEINLINE bool IsTransform(int32 Element) const { return !IsGeometry(Element); }
 	FORCEINLINE void SetFlags(int32 Element, int32 InFlags) { TManagedArray<int32>& Status = StatusFlags; Status[Element] |= InFlags; }
 	FORCEINLINE void ClearFlags(int32 Element, int32 InFlags) { TManagedArray<int32>& Status = StatusFlags; Status[Element] = Status[Element] & ~InFlags; }
diff --git a/Engine/Source/Runtime/Experimental/Chaos/Public/GeometryCollection/GeometryCollectionUtility.h b/Engine/Source/Runtime/Experimental/Chaos/Public/GeometryCollection/GeometryCollectionUtility.h
--- a/Engine/Source/Runtime/Experimental/Chaos/Public/GeometryCollection/GeometryCollectionUtility.h
+++ b/Engine/Source/Runtime/Experimental/Chaos/Public/GeometryCollection/GeometryCollectionUtility.h
@@ -46,10 +46,12 @@ namespace GeometryCollection
 		TManagedArray<int32>& MaterialIndex = RestCollection->MaterialIndex;
 		TManagedArray<int32>& MaterialID = RestCollection->MaterialID;
 		TManagedArray<FTransform>& Transform = RestCollection->Transform;
+		TManagedArray<int32>& SimulationType = RestCollection->SimulationType;
 
 		// Set particle info
 		Transform[0] = Xf;
 		Transform[0].NormalizeRotation();
+		SimulationType[0] = FGeometryCollection::ESimulationTypes::FST_Rigid;
 
 		// Set vertex info
 		for (int32 Idx = 0; Idx < PointsIn.Num(); ++Idx)
