GeometryCollection
  - Removed SimulationType FST_Non_Simulatable.
  - Non simulated embedded geometry nodes are now tagged with SimulationType FST_None.

Fracture Mode
  - PlanarCuts applies appropriate SimulationType to new transforms and clusters.

diff --git a/Engine/Plugins/Experimental/ChaosEditor/Source/FractureEditor/Private/FractureToolGenerators.cpp b/Engine/Plugins/Experimental/ChaosEditor/Source/FractureEditor/Private/FractureToolGenerators.cpp
--- a/Engine/Plugins/Experimental/ChaosEditor/Source/FractureEditor/Private/FractureToolGenerators.cpp
+++ b/Engine/Plugins/Experimental/ChaosEditor/Source/FractureEditor/Private/FractureToolGenerators.cpp
@@ -464,7 +464,7 @@ void UFractureToolAddEmbeddedGeometry::Execute(TWeakPtr<FFractureEditorModeToolk
 				Children[SelectedBone].Add(NewTransformIndex);
 
 				// Set properties appropriate to embedded geometry
-				SimType[NewTransformIndex] = FGeometryCollection::ESimulationTypes::FST_Not_Simulatable;
+				SimType[NewTransformIndex] = FGeometryCollection::ESimulationTypes::FST_None;
 				CollectionSimulatableParticles[NewTransformIndex] = false;
 			}	
 		}
diff --git a/Engine/Plugins/Experimental/ChaosEditor/Source/FractureEditor/Private/SGeometryCollectionOutliner.cpp b/Engine/Plugins/Experimental/ChaosEditor/Source/FractureEditor/Private/SGeometryCollectionOutliner.cpp
--- a/Engine/Plugins/Experimental/ChaosEditor/Source/FractureEditor/Private/SGeometryCollectionOutliner.cpp
+++ b/Engine/Plugins/Experimental/ChaosEditor/Source/FractureEditor/Private/SGeometryCollectionOutliner.cpp
@@ -420,13 +420,30 @@ TSharedRef<ITableRow> FGeometryCollectionTreeItemBone::MakeTreeRowWidget(const T
 		ItemText = FText::FromString(FString::FromInt(GetBoneIndex()));
 	}
 
-	// Cyan indicates cluster, Gray is leaf node	
+	// Set color according to simulation type
+
+	FSlateColor TextColor(FLinearColor::Red); // default color indicates something wrong
+
 	TSharedPtr<FGeometryCollection, ESPMode::ThreadSafe> GeometryCollectionPtr = ParentComponentItem->GetComponent()->GetRestCollection()->GetGeometryCollection();
-	const TManagedArray<TSet<int32>>& Children = GeometryCollectionPtr->GetAttribute<TSet<int32>>("Children", FGeometryCollection::TransformGroup);
-	FSlateColor TextColor(FLinearColor::Gray);
-	if (Children[GetBoneIndex()].Num() > 0)
+	const TManagedArray<int32>& SimulationType = GeometryCollectionPtr->SimulationType;
+	switch (SimulationType[GetBoneIndex()])
 	{
-		TextColor = FColor::Cyan;
+		// This is a dubious case. We shouldn't expect to see nodes with this type.
+		case FGeometryCollection::ESimulationTypes::FST_None:
+			TextColor = FColor::Green;
+			break;
+
+		case FGeometryCollection::ESimulationTypes::FST_Rigid:
+			TextColor = FLinearColor::Gray;
+			break;
+
+		case FGeometryCollection::ESimulationTypes::FST_Clustered:
+			TextColor = FColor::Cyan;
+			break;
+
+		default:
+			ensureMsgf(false, TEXT("Invalid Geometry Collection simulation type encountered."));
+			break;
 	}
 	
 	return SNew(STableRow<FGeometryCollectionTreeItemPtr>, InOwnerTable)
diff --git a/Engine/Plugins/Experimental/PlanarCutPlugin/Source/PlanarCut/Private/PlanarCut.cpp b/Engine/Plugins/Experimental/PlanarCutPlugin/Source/PlanarCut/Private/PlanarCut.cpp
--- a/Engine/Plugins/Experimental/PlanarCutPlugin/Source/PlanarCut/Private/PlanarCut.cpp
+++ b/Engine/Plugins/Experimental/PlanarCutPlugin/Source/PlanarCut/Private/PlanarCut.cpp
@@ -1894,8 +1894,10 @@ struct FDynamicMeshCollection
 			Output.BoneColor[TransformIdx] = Output.BoneColor[TransformParent];
 			Output.Parent[TransformIdx] = TransformParent;
 			Output.Children[TransformParent].Add(TransformIdx);
+			Output.SimulationType[TransformParent] = FGeometryCollection::ESimulationTypes::FST_Clustered;
 		}
 		Output.Transform[TransformIdx] = FTransform::Identity;
+		Output.SimulationType[TransformIdx] = FGeometryCollection::ESimulationTypes::FST_Rigid;
 
 		int32 FacesStart = Output.AddElements(NumTriangles, FGeometryCollection::FacesGroup);
 		int32 VerticesStart = Output.AddElements(NumVertices, FGeometryCollection::VerticesGroup);
diff --git a/Engine/Source/Runtime/Experimental/Chaos/Private/PhysicsProxy/GeometryCollectionPhysicsProxy.cpp b/Engine/Source/Runtime/Experimental/Chaos/Private/PhysicsProxy/GeometryCollectionPhysicsProxy.cpp
--- a/Engine/Source/Runtime/Experimental/Chaos/Private/PhysicsProxy/GeometryCollectionPhysicsProxy.cpp
+++ b/Engine/Source/Runtime/Experimental/Chaos/Private/PhysicsProxy/GeometryCollectionPhysicsProxy.cpp
@@ -1790,7 +1790,7 @@ void FGeometryCollectionPhysicsProxy::InitializeSharedCollisionStructures(
 	for (int32 GeometryIndex = 0; GeometryIndex < NumGeometries; GeometryIndex++)
 	{
 		const int32 TransformGroupIndex = TransformIndex[GeometryIndex];
-		if (SimulationType[TransformGroupIndex]!= FGeometryCollection::ESimulationTypes::FST_Not_Simulatable)
+		if (SimulationType[TransformGroupIndex] > FGeometryCollection::ESimulationTypes::FST_None)
 		{
 			TUniquePtr<TTriangleMesh<float>> TriMesh(
 				CreateTriangleMesh(
diff --git a/Engine/Source/Runtime/Experimental/Chaos/Public/GeometryCollection/GeometryCollection.h b/Engine/Source/Runtime/Experimental/Chaos/Public/GeometryCollection/GeometryCollection.h
--- a/Engine/Source/Runtime/Experimental/Chaos/Public/GeometryCollection/GeometryCollection.h
+++ b/Engine/Source/Runtime/Experimental/Chaos/Public/GeometryCollection/GeometryCollection.h
@@ -96,8 +96,7 @@ public:
 	{
 		FST_None = 0,
 		FST_Rigid = 1,
-		FST_Clustered = 2,
-		FST_Not_Simulatable = 3
+		FST_Clustered = 2
 	};
 
 	enum ENodeFlags : uint32
