Fracture Mode -- Implemented context filters: Embedded and Rigid. Previously selection filters for rigids were using the lead node condition to determine if selected node is rigid. Now built to use proper tags.

diff --git a/Engine/Plugins/Experimental/ChaosEditor/Source/FractureEditor/Private/FractureToolContext.cpp b/Engine/Plugins/Experimental/ChaosEditor/Source/FractureEditor/Private/FractureToolContext.cpp
--- a/Engine/Plugins/Experimental/ChaosEditor/Source/FractureEditor/Private/FractureToolContext.cpp
+++ b/Engine/Plugins/Experimental/ChaosEditor/Source/FractureEditor/Private/FractureToolContext.cpp
@@ -54,7 +54,6 @@ TMap<int32, TArray<int32>> FFractureToolContext::GetClusteredSelections()
 	return SiblingGroups;
 }
 
-
 void FFractureToolContext::ConvertSelectionToLeafNodes()
 {
 	Sanitize();
@@ -85,19 +84,86 @@ void FFractureToolContext::ConvertSelectionToLeafNodes(int32 Index, TArray<int32
 }
 
 
+void FFractureToolContext::ConvertSelectionToRigidNodes()
+{
+	Sanitize();
+
+	TArray<int32> RigidSelection;
+	for (int32 Index : SelectedBones)
+	{
+		ConvertSelectionToRigidNodes(Index, RigidSelection);
+	}
+
+	SelectedBones = RigidSelection;
+}
+
+void FFractureToolContext::ConvertSelectionToRigidNodes(int32 Index, TArray<int32>& RigidSelection)
+{
+	const TManagedArray<TSet<int32>>& Children = GeometryCollection->Children;
+	const TManagedArray<int32>& SimulationType = GeometryCollection->SimulationType;
+
+	if (SimulationType[Index] == FGeometryCollection::ESimulationTypes::FST_Rigid)
+	{
+		RigidSelection.Add(Index);
+	}
+	else
+	{
+		for (int32 Child : Children[Index])
+		{
+			ConvertSelectionToRigidNodes(Child, RigidSelection);
+		}
+	}
+}
+
+void FFractureToolContext::ConvertSelectionToEmbeddedGeometryNodes()
+{
+	Sanitize();
+
+	TArray<int32> EmbeddedSelection;
+	for (int32 Index : SelectedBones)
+	{
+		ConvertSelectionToEmbeddedGeometryNodes(Index, EmbeddedSelection);
+	}
+
+	SelectedBones = EmbeddedSelection;
+}
+
+void FFractureToolContext::ConvertSelectionToEmbeddedGeometryNodes(int32 Index, TArray<int32>& EmbeddedSelection)
+{
+	const TManagedArray<TSet<int32>>& Children = GeometryCollection->Children;
+	const TManagedArray<int32>& SimulationType = GeometryCollection->SimulationType;
+
+	if (SimulationType[Index] == FGeometryCollection::ESimulationTypes::FST_None)
+	{
+		EmbeddedSelection.Add(Index);
+	}
+	else
+	{
+		for (int32 Child : Children[Index])
+		{
+			ConvertSelectionToEmbeddedGeometryNodes(Child, EmbeddedSelection);
+		}
+	}
+}
+
 void FFractureToolContext::ConvertSelectionToClusterNodes()
 {
-	// If this is a leaf node, select the cluster instead.
-	const TManagedArray<TSet<int32>>& Children = GeometryCollection->GetAttribute<TSet<int32>>("Children", FGeometryCollection::TransformGroup);
-	const TManagedArray<int32>& Parents = GeometryCollection->GetAttribute<int32>("Parent", FGeometryCollection::TransformGroup);
+	// If this is a non-cluster node, select the cluster containing it instead.
+	const TManagedArray<TSet<int32>>& Children = GeometryCollection->Children;
+	const TManagedArray<int32>& Parents = GeometryCollection->Parent;
+	const TManagedArray<int32>& SimulationType = GeometryCollection->SimulationType;
 
 	TArray<int32> AddedClusterSelections;
 	for (int32 Index : SelectedBones)
 	{
-		if (Children[Index].Num() == 0)
+		if (SimulationType[Index] == FGeometryCollection::ESimulationTypes::FST_Rigid)
 		{
 			AddedClusterSelections.AddUnique(Parents[Index]);
 		}
+		else if (SimulationType[Index] == FGeometryCollection::ESimulationTypes::FST_None)
+		{
+			AddedClusterSelections.AddUnique(Parents[Parents[Index]]);
+		}
 	}
 	SelectedBones.Append(AddedClusterSelections);
 	
diff --git a/Engine/Plugins/Experimental/ChaosEditor/Source/FractureEditor/Private/FractureToolContext.h b/Engine/Plugins/Experimental/ChaosEditor/Source/FractureEditor/Private/FractureToolContext.h
--- a/Engine/Plugins/Experimental/ChaosEditor/Source/FractureEditor/Private/FractureToolContext.h
+++ b/Engine/Plugins/Experimental/ChaosEditor/Source/FractureEditor/Private/FractureToolContext.h
@@ -20,6 +20,8 @@ public:
 
 	void Sanitize();
 	void ConvertSelectionToLeafNodes();
+	void ConvertSelectionToRigidNodes();
+	void ConvertSelectionToEmbeddedGeometryNodes();
 	void ConvertSelectionToClusterNodes();
 	void RemoveRootNodes();
 	
@@ -46,6 +48,8 @@ public:
 private:
 	bool HasSelectedAncestor(int32 Index) const;
 	void ConvertSelectionToLeafNodes(int32 Index, TArray<int32>& LeafSelection);
+	void ConvertSelectionToRigidNodes(int32 Index, TArray<int32>& RigidSelection);
+	void ConvertSelectionToEmbeddedGeometryNodes(int32 Index, TArray<int32>& EmbeddedSelection);
 
 private:
 	TArray<int32> SelectedBones;
diff --git a/Engine/Plugins/Experimental/ChaosEditor/Source/FractureEditor/Private/FractureToolCutter.cpp b/Engine/Plugins/Experimental/ChaosEditor/Source/FractureEditor/Private/FractureToolCutter.cpp
--- a/Engine/Plugins/Experimental/ChaosEditor/Source/FractureEditor/Private/FractureToolCutter.cpp
+++ b/Engine/Plugins/Experimental/ChaosEditor/Source/FractureEditor/Private/FractureToolCutter.cpp
@@ -50,7 +50,7 @@ TArray<FFractureToolContext> UFractureToolCutterBase::GetFractureToolContexts()
 	{
 		// Generate a context for each selected node
 		FFractureToolContext FullSelection(GeometryCollectionComponent);
-		FullSelection.ConvertSelectionToLeafNodes();
+		FullSelection.ConvertSelectionToRigidNodes();
 		
 		// Update global transforms and bounds		
 		const TManagedArray<FTransform>& Transform = FullSelection.GetGeometryCollection()->GetAttribute<FTransform>("Transform", FGeometryCollection::TransformGroup);
diff --git a/Engine/Plugins/Experimental/ChaosEditor/Source/FractureEditor/Private/FractureToolGenerators.cpp b/Engine/Plugins/Experimental/ChaosEditor/Source/FractureEditor/Private/FractureToolGenerators.cpp
--- a/Engine/Plugins/Experimental/ChaosEditor/Source/FractureEditor/Private/FractureToolGenerators.cpp
+++ b/Engine/Plugins/Experimental/ChaosEditor/Source/FractureEditor/Private/FractureToolGenerators.cpp
@@ -432,7 +432,7 @@ void UFractureToolAddEmbeddedGeometry::Execute(TWeakPtr<FFractureEditorModeToolk
 
 		const FTransform TargetActorTransform(Context.GetGeometryCollectionComponent()->GetOwner()->GetTransform());
 		
-		Context.ConvertSelectionToLeafNodes();
+		Context.ConvertSelectionToRigidNodes();
 		const TArray<int32>& SelectedBones = Context.GetSelection();
 		for (const int32 SelectedBone : SelectedBones)
 		{
