FractureEditor -- Fracture settings random seed == -1 now produces different results each time the fracture is applied.

diff --git a/Engine/Plugins/Experimental/ChaosEditor/Source/FractureEditor/Private/FractureToolCutter.cpp b/Engine/Plugins/Experimental/ChaosEditor/Source/FractureEditor/Private/FractureToolCutter.cpp
--- a/Engine/Plugins/Experimental/ChaosEditor/Source/FractureEditor/Private/FractureToolCutter.cpp
+++ b/Engine/Plugins/Experimental/ChaosEditor/Source/FractureEditor/Private/FractureToolCutter.cpp
@@ -72,7 +72,7 @@ TArray<FFractureToolContext> UFractureToolCutterBase::GetFractureToolContexts()
 
 		if (CutterSettings->bGroupFracture)
 		{
-			FullSelection.SetSeed(CutterSettings->RandomSeed == -1 ? CutterSettings->RandomSeed : FMath::Rand());
+			FullSelection.SetSeed(CutterSettings->RandomSeed > -1 ? CutterSettings->RandomSeed : FMath::Rand());
 
 			FBox Bounds(ForceInit);
 			for (int32 BoneIndex : FullSelection.GetSelection())
@@ -94,7 +94,7 @@ TArray<FFractureToolContext> UFractureToolCutterBase::GetFractureToolContexts()
 				Contexts.Emplace(GeometryCollectionComponent);
 				FFractureToolContext& FractureContext = Contexts.Last();
 
-				FractureContext.SetSeed(CutterSettings->RandomSeed == -1 ? CutterSettings->RandomSeed + Index : FMath::Rand());
+				FractureContext.SetSeed(CutterSettings->RandomSeed > -1 ? CutterSettings->RandomSeed + Index : FMath::Rand());
 				FractureContext.SetBounds(BoundsToBone[Index]);
 			}
 		}
