Fracture Mode
* Fixed some longstanding issues with histogram displaying approariate level view.
* Fixed problem with bone selection immediately following fracture crshing due to stale outliner/histogram views.
* Histogram no longer attempt to display embedded geometry data.
* Hover text on Histogram now displays bone index as well as inspected value.
* Embedded geometry broken out into it's own category in Level Statistics.
* Corrected clustering behavior with respect to embedded geometry nodes

diff --git a/Engine/Plugins/Experimental/ChaosEditor/Source/FractureEditor/Private/FractureEditorModeToolkit.cpp b/Engine/Plugins/Experimental/ChaosEditor/Source/FractureEditor/Private/FractureEditorModeToolkit.cpp
--- a/Engine/Plugins/Experimental/ChaosEditor/Source/FractureEditor/Private/FractureEditorModeToolkit.cpp
+++ b/Engine/Plugins/Experimental/ChaosEditor/Source/FractureEditor/Private/FractureEditorModeToolkit.cpp
@@ -936,7 +936,7 @@ void FFractureEditorModeToolkit::SetOutlinerComponents(const TArray<UGeometryCol
 
 	if (HistogramView)
 	{
-		HistogramView->SetComponents(InNewComponents);
+		HistogramView->SetComponents(InNewComponents, GetLevelViewValue());
 	}
 
 	if (ActiveTool != nullptr)
@@ -1112,6 +1112,7 @@ void FFractureEditorModeToolkit::UpdateGeometryComponentAttributes(UGeometryColl
 void FFractureEditorModeToolkit::UpdateVolumes(FGeometryCollectionPtr GeometryCollection, const Chaos::TParticles<float, 3>& MassSpaceParticles, int32 TransformIndex)
 {
 	const TManagedArray<TSet<int32>>& Children = GeometryCollection->Children;
+	const TManagedArray<int32>& SimulationType = GeometryCollection->SimulationType;
 	const TManagedArray<bool>& Visible = GeometryCollection->Visible;
 	const TManagedArray<int32>& FaceCount = GeometryCollection->FaceCount;
 	const TManagedArray<int32>& FaceStart = GeometryCollection->FaceStart;
@@ -1119,7 +1120,7 @@ void FFractureEditorModeToolkit::UpdateVolumes(FGeometryCollectionPtr GeometryCo
 	const TManagedArray<FIntVector>& Indices = GeometryCollection->Indices;
 	TManagedArray<float>& Volumes = GeometryCollection->GetAttribute<float>("Volume", FTransformCollection::TransformGroup);
 
-	if (Children[TransformIndex].Num() == 0) // leaf node
+	if (SimulationType[TransformIndex] == FGeometryCollection::ESimulationTypes::FST_Rigid)
 	{
 		int32 GeometryIndex = TransformToGeometryIndex[TransformIndex];
 
@@ -1141,7 +1142,7 @@ void FFractureEditorModeToolkit::UpdateVolumes(FGeometryCollectionPtr GeometryCo
 			Volumes[TransformIndex] = Volume;
 		}	
 	}
-	else
+	else if (SimulationType[TransformIndex] == FGeometryCollection::ESimulationTypes::FST_Clustered)
 	{
 		// Recurse to children and sum the volumes for this node
 		float LocalVolume = 0.0;
@@ -1153,6 +1154,11 @@ void FFractureEditorModeToolkit::UpdateVolumes(FGeometryCollectionPtr GeometryCo
 
 		Volumes[TransformIndex] = LocalVolume;
 	}
+	else
+	{
+		// Node is embedded geometry, for which we do not require volume calculations.
+		Volumes[TransformIndex] = 0.0;
+	}
 }
 
 bool GetValidGeoCenter(const TManagedArray<int32>& TransformToGeometryIndex, const TArray<FTransform>& Transforms, const TManagedArray<TSet<int32>>& Children, const TManagedArray<FBox>& BoundingBox, int32 TransformIndex, FVector& OutGeoCenter )
@@ -1297,9 +1303,14 @@ void FFractureEditorModeToolkit::UpdateExplodedVectors(UGeometryCollectionCompon
 	}
 }
 
+void FFractureEditorModeToolkit::RegenerateOutliner()
+{
+	OutlinerView->UpdateGeometryCollection();
+}
+
 void FFractureEditorModeToolkit::RegenerateHistogram()
 {
-	HistogramView->RegenerateNodes();
+	HistogramView->RegenerateNodes(GetLevelViewValue());
 }
 
 void FFractureEditorModeToolkit::OnOutlinerBoneSelectionChanged(UGeometryCollectionComponent* RootComponent, TArray<int32>& SelectedBones)
@@ -1365,6 +1376,7 @@ FText FFractureEditorModeToolkit::GetStatisticsSummary() const
 		TArray<int32> LevelTransformsAll;
 		LevelTransformsAll.SetNumZeroed(10);
 		int32 LevelMax = INT_MIN;
+		int32 EmbeddedCount = 0;
 
 		for (int32 Idx = 0; Idx < GeometryCollectionArray.Num(); ++Idx)
 		{
@@ -1378,17 +1390,25 @@ FText FFractureEditorModeToolkit::GetStatisticsSummary() const
 			if(GeometryCollection->HasAttribute("Level", FGeometryCollection::TransformGroup))
 			{
 				const TManagedArray<int32>& Levels = GeometryCollection->GetAttribute<int32>("Level", FGeometryCollection::TransformGroup);
+				const TManagedArray<int32>& SimulationType = GeometryCollection->SimulationType;
 
 				TArray<int32> LevelTransforms;
 				for(int32 Element = 0, NumElement = Levels.Num(); Element < NumElement; ++Element)
 				{
-					const int32 NodeLevel = Levels[Element];
-					while(LevelTransforms.Num() <= NodeLevel)
+					if (SimulationType[Element] == FGeometryCollection::ESimulationTypes::FST_None)
 					{
-						LevelTransforms.SetNum(NodeLevel + 1);
-						LevelTransforms[NodeLevel] = 0;
+						++EmbeddedCount;
 					}
-					++LevelTransforms[NodeLevel];
+					else
+					{
+						const int32 NodeLevel = Levels[Element];
+						while (LevelTransforms.Num() <= NodeLevel)
+						{
+							LevelTransforms.SetNum(NodeLevel + 1);
+							LevelTransforms[NodeLevel] = 0;
+						}
+						++LevelTransforms[NodeLevel];
+					}		
 				}
 
 				for(int32 Level = 0; Level < LevelTransforms.Num(); ++Level)
@@ -1407,6 +1427,7 @@ FText FFractureEditorModeToolkit::GetStatisticsSummary() const
 		{
 			Buffer += FString::Printf(TEXT("Level: %d \t - \t %d\n"), Level, LevelTransformsAll[Level]);
 		}
+		Buffer += FString::Printf(TEXT("\nEmbedded: %d"), EmbeddedCount);
 	}
 
 	return FText::FromString(Buffer);
diff --git a/Engine/Plugins/Experimental/ChaosEditor/Source/FractureEditor/Private/FractureToolClusterMagnet.cpp b/Engine/Plugins/Experimental/ChaosEditor/Source/FractureEditor/Private/FractureToolClusterMagnet.cpp
--- a/Engine/Plugins/Experimental/ChaosEditor/Source/FractureEditor/Private/FractureToolClusterMagnet.cpp
+++ b/Engine/Plugins/Experimental/ChaosEditor/Source/FractureEditor/Private/FractureToolClusterMagnet.cpp
@@ -74,8 +74,6 @@ void UFractureToolClusterMagnet::Execute(TWeakPtr<FFractureEditorModeToolkit> In
 	{
 		FFractureEditorModeToolkit* Toolkit = InToolkit.Pin().Get();
 
-		int32 CurrentLevelView = Toolkit->GetLevelViewValue();
-
 		TArray<FFractureToolContext> Contexts = GetFractureToolContexts();
 
 		for (FFractureToolContext& Context : Contexts)
@@ -94,9 +92,6 @@ void UFractureToolClusterMagnet::Execute(TWeakPtr<FFractureEditorModeToolkit> In
 
 			for (TPair<int32, TArray<int32>>& Group : ClusteredSelection)
 			{
-				const TArray<int32>& CurrentSelection = Group.Value;
-				CurrentLevelView = Levels[Group.Key] + 1;
-
 				// We have the connections for the leaf nodes of our geometry collection. We want to percolate those up to the top nodes.
 				TMap<int32, TSet<int32>> TopNodeConnectivity = InitializeConnectivity(Children[Group.Key], Context.GetGeometryCollection(), Levels[Group.Key]+1);
 
@@ -138,12 +133,6 @@ void UFractureToolClusterMagnet::Execute(TWeakPtr<FFractureEditorModeToolkit> In
 			Refresh(Context, Toolkit);
 		}
 
-		// Set the current viewing level to our best guess
-		if (CurrentLevelView != Toolkit->GetLevelViewValue())
-		{
-			Toolkit->OnSetLevelViewValue(CurrentLevelView);
-		}
-
 		SetOutlinerComponents(Contexts, Toolkit);
 	}
 }
diff --git a/Engine/Plugins/Experimental/ChaosEditor/Source/FractureEditor/Private/FractureToolClustering.cpp b/Engine/Plugins/Experimental/ChaosEditor/Source/FractureEditor/Private/FractureToolClustering.cpp
--- a/Engine/Plugins/Experimental/ChaosEditor/Source/FractureEditor/Private/FractureToolClustering.cpp
+++ b/Engine/Plugins/Experimental/ChaosEditor/Source/FractureEditor/Private/FractureToolClustering.cpp
@@ -38,8 +38,6 @@ void UFractureToolFlattenAll::Execute(TWeakPtr<FFractureEditorModeToolkit> InToo
 	{
 		FFractureEditorModeToolkit* Toolkit = InToolkit.Pin().Get();
 
-		int32 CurrentLevelView = Toolkit->GetLevelViewValue();
-
 		TArray<FFractureToolContext> Contexts = GetFractureToolContexts();
 
 		for (FFractureToolContext& Context : Contexts)
@@ -57,17 +55,11 @@ void UFractureToolFlattenAll::Execute(TWeakPtr<FFractureEditorModeToolkit> InToo
 				// Cleanup: Remove any clusters remaining in the flattened branch.
 				FGeometryCollectionClusteringUtility::RemoveDanglingClusters(Context.GetGeometryCollection().Get());
 
-				CurrentLevelView = Levels[ClusterIndex] + 1;
 			}
 
 			Refresh(Context, Toolkit);
 		}
 
-		if (CurrentLevelView != Toolkit->GetLevelViewValue())
-		{
-			Toolkit->OnSetLevelViewValue(CurrentLevelView);
-		}
-
 		SetOutlinerComponents(Contexts, Toolkit);
 	}
 }
@@ -168,8 +160,6 @@ void UFractureToolUncluster::Execute(TWeakPtr<FFractureEditorModeToolkit> InTool
 	{
 		FFractureEditorModeToolkit* Toolkit = InToolkit.Pin().Get();
 
-		int32 CurrentLevelView = Toolkit->GetLevelViewValue();
-
 		TArray<FFractureToolContext> Contexts = GetFractureToolContexts();
 
 		for (FFractureToolContext& Context : Contexts)
@@ -189,19 +179,10 @@ void UFractureToolUncluster::Execute(TWeakPtr<FFractureEditorModeToolkit> InTool
 
 			FGeometryCollectionClusteringUtility::CollapseHierarchyOneLevel(Context.GetGeometryCollection().Get(), Context.GetSelection());
 			Context.SetSelection(NewSelection);
-			if (NewSelection.Num() > 0)
-			{
-				CurrentLevelView = Levels[NewSelection.Last()];
-			}
 			
 			Refresh(Context, Toolkit);
 		}
 
-		if (CurrentLevelView != Toolkit->GetLevelViewValue())
-		{
-			Toolkit->OnSetLevelViewValue(CurrentLevelView);
-		}
-
 		SetOutlinerComponents(Contexts, Toolkit);
 	}
 }
@@ -239,7 +220,7 @@ void UFractureToolMoveUp::Execute(TWeakPtr<FFractureEditorModeToolkit> InToolkit
 
 		for (FFractureToolContext& Context: Contexts)
 		{
-			Context.Sanitize();
+			Context.ConvertSelectionToRigidNodes();
 			FGeometryCollectionClusteringUtility::MoveUpOneHierarchyLevel(Context.GetGeometryCollection().Get(), Context.GetSelection());
 			Refresh(Context, Toolkit);
 		}
diff --git a/Engine/Plugins/Experimental/ChaosEditor/Source/FractureEditor/Private/SGeometryCollectionHistogram.cpp b/Engine/Plugins/Experimental/ChaosEditor/Source/FractureEditor/Private/SGeometryCollectionHistogram.cpp
--- a/Engine/Plugins/Experimental/ChaosEditor/Source/FractureEditor/Private/SGeometryCollectionHistogram.cpp
+++ b/Engine/Plugins/Experimental/ChaosEditor/Source/FractureEditor/Private/SGeometryCollectionHistogram.cpp
@@ -35,7 +35,7 @@ TSharedRef<ITableRow> FGeometryCollectionHistogramItem::MakeHistogramRowWidget(c
 				.Color(NodeColor)
 				.AlphaDisplayMode(EColorBlockAlphaDisplayMode::Ignore)
 				.Size(FVector2D(1.0f,5.0f))
-				.ToolTipText(HoverText)
+				.ToolTipText(FText::FromString(HoverString))
 			]
 			+ SHorizontalBox::Slot()
 			.FillWidth(1.0-NormalizedValue)
@@ -59,13 +59,13 @@ void FGeometryCollectionHistogramItem::SetInspectedAttribute(EInspectedAttribute
 			{
 			case EInspectedAttributeEnum::Volume:
 				InspectedValue = Collection->GetAttribute<float>(TEXT("Volume"), TEXT("Transform"))[BoneIndex];
-				HoverText = FText::FromString(FString::SanitizeFloat(InspectedValue));
+				HoverString = FString::Printf(TEXT("%d: %.2f"), BoneIndex, InspectedValue);
 				break;
 
 			case EInspectedAttributeEnum::Level:
 				{
 					int32 Level = Collection->GetAttribute<int32>(TEXT("Level"), TEXT("Transform"))[BoneIndex];
-					HoverText = FText::FromString(FString::FormatAsNumber(Level));
+					HoverString = FString::Printf(TEXT("%d: %d"), BoneIndex, Level);
 					InspectedValue = static_cast<float>(Level);
 				}
 				break;
@@ -75,7 +75,7 @@ void FGeometryCollectionHistogramItem::SetInspectedAttribute(EInspectedAttribute
 					int32 InitialDynamicState = static_cast<int32>(Collection->GetAttribute<int32>(TEXT("InitialDynamicState"), TEXT("Transform"))[BoneIndex]);
 					
 					static const TArray<FString> HoverNames{ "No Override", "Sleeping", "Kinematic", "Static" };
-					HoverText = FText::FromString(HoverNames[InitialDynamicState]);
+					HoverString = FString::Printf(TEXT("%d: %s"), BoneIndex, *HoverNames[InitialDynamicState]);
 
 					InspectedValue = static_cast<float>(InitialDynamicState);
 				}
@@ -116,7 +116,7 @@ FGeometryCollectionHistogramItemPtr FGeometryCollectionHistogramItemComponent::G
 	return FGeometryCollectionHistogramItemPtr();
 }
 
-FGeometryCollectionHistogramItemList FGeometryCollectionHistogramItemComponent::RegenerateNodes()
+FGeometryCollectionHistogramItemList FGeometryCollectionHistogramItemComponent::RegenerateNodes(int32 LevelView)
 {
 	// Collect the inspected attribute 
 	FGeometryCollectionHistogramItemList NodesList;
@@ -139,8 +139,21 @@ FGeometryCollectionHistogramItemList FGeometryCollectionHistogramItemComponent::
 			// Add a sub item to the histogram for each of the geometry nodes in this GeometryCollection
 			for (int32 Index = 0; Index < NumElements; Index++)
 			{
-				if (Collection->Children[Index].Num() == 0)
+				if (Collection->SimulationType[Index] == FGeometryCollection::ESimulationTypes::FST_Rigid)
 				{
+					if (LevelView > -1)
+					{
+						// Don't display nodes that aren't visible at the currently inspected level.
+						if (Collection->HasAttribute("Level", FGeometryCollection::TransformGroup))
+						{
+							const TManagedArray<int32>& Levels = Collection->GetAttribute<int32>("Level", FGeometryCollection::TransformGroup);
+							if (Levels[Index] != LevelView)
+							{
+								continue;
+							}
+						}
+					}
+
 					TSharedRef<FGeometryCollectionHistogramItem> NewItem = MakeShared<FGeometryCollectionHistogramItem>(Guids[Index], Index, AsShared());
 
 					FLinearColor GeoColor = BoneColor[Index];
@@ -172,7 +185,7 @@ void SGeometryCollectionHistogram::Construct(const FArguments& InArgs)
 	];
 }
 
-void SGeometryCollectionHistogram::SetComponents(const TArray<UGeometryCollectionComponent*>& InNewComponents)
+void SGeometryCollectionHistogram::SetComponents(const TArray<UGeometryCollectionComponent*>& InNewComponents, int32 LevelView)
 {
 	// Clear the cached Tree ItemSelection without affecting the SelectedBones as 
 	// we want to refresh the tree selection using selected bones
@@ -187,7 +200,7 @@ void SGeometryCollectionHistogram::SetComponents(const TArray<UGeometryCollectio
 		for (UGeometryCollectionComponent* Component : InNewComponents)
 		{
 			RootNodes.Add(MakeShared<FGeometryCollectionHistogramItemComponent>(Component));
-			LeafNodes.Append(RootNodes.Last()->RegenerateNodes());
+			LeafNodes.Append(RootNodes.Last()->RegenerateNodes(LevelView));
 		}
 
 		SetListIndices();
@@ -240,11 +253,11 @@ void SGeometryCollectionHistogram::RefreshView(bool bSorted)
 	ListView->RebuildList();
 }
 
-void SGeometryCollectionHistogram::RegenerateNodes()
+void SGeometryCollectionHistogram::RegenerateNodes(int32 LevelView)
 {
 	for (TSharedPtr<FGeometryCollectionHistogramItemComponent> Root : RootNodes)
 	{
-		Root->RegenerateNodes();
+		Root->RegenerateNodes(LevelView);
 	}
 
 	UHistogramSettings* HistogramSettings = GetMutableDefault<UHistogramSettings>();
@@ -297,22 +310,25 @@ void SGeometryCollectionHistogram::OnSelectionChanged(FGeometryCollectionHistogr
 
 void SGeometryCollectionHistogram::NormalizeInspectedValues()
 {
-	// Sort nodes by value
-	LeafNodes.Sort([](const TSharedPtr<FGeometryCollectionHistogramItem> A, const TSharedPtr<FGeometryCollectionHistogramItem> B)
-		{
-			return A->GetInspectedValue() > B->GetInspectedValue();
-		});
+	if (LeafNodes.Num())
+	{
+		// Sort nodes by value
+		LeafNodes.Sort([](const TSharedPtr<FGeometryCollectionHistogramItem> A, const TSharedPtr<FGeometryCollectionHistogramItem> B)
+			{
+				return A->GetInspectedValue() > B->GetInspectedValue();
+			});
 
-	// Normalize
-	float MaxInspectedValue = LeafNodes[0]->GetInspectedValue();
-	float MinInspectedValue = LeafNodes.Last()->GetInspectedValue();
+		// Normalize
+		float MaxInspectedValue = LeafNodes[0]->GetInspectedValue();
+		float MinInspectedValue = LeafNodes.Last()->GetInspectedValue();
 
-	// We allow a little extra room at the bottom so that bars don't drop to zero length
-	MinInspectedValue -= (MaxInspectedValue - MinInspectedValue) * 0.025;
+		// We allow a little extra room at the bottom so that bars don't drop to zero length
+		MinInspectedValue -= (MaxInspectedValue - MinInspectedValue) * 0.025;
 
-	for (TSharedPtr<FGeometryCollectionHistogramItem> LeafNode : LeafNodes)
-	{
-		LeafNode->SetNormalizedValue(MinInspectedValue, MaxInspectedValue);
+		for (TSharedPtr<FGeometryCollectionHistogramItem> LeafNode : LeafNodes)
+		{
+			LeafNode->SetNormalizedValue(MinInspectedValue, MaxInspectedValue);
+		}
 	}
 }
 
diff --git a/Engine/Plugins/Experimental/ChaosEditor/Source/FractureEditor/Private/SGeometryCollectionHistogram.h b/Engine/Plugins/Experimental/ChaosEditor/Source/FractureEditor/Private/SGeometryCollectionHistogram.h
--- a/Engine/Plugins/Experimental/ChaosEditor/Source/FractureEditor/Private/SGeometryCollectionHistogram.h
+++ b/Engine/Plugins/Experimental/ChaosEditor/Source/FractureEditor/Private/SGeometryCollectionHistogram.h
@@ -63,7 +63,7 @@ public:
 
 	FGeometryCollectionHistogramItemPtr GetItemFromBoneIndex(int32 BoneIndex) const;
 	
-	FGeometryCollectionHistogramItemList RegenerateNodes();
+	FGeometryCollectionHistogramItemList RegenerateNodes(int32 LevelView);
 
 private:
 	TWeakObjectPtr<UGeometryCollectionComponent> Component;
@@ -84,7 +84,6 @@ public:
 		, NodeColor(FLinearColor::Black)
 		, NormalizedValue(0.0)
 		, InspectedValue(0.0)
-		, HoverText(FText())
 	{}
 
 	TSharedRef<ITableRow> MakeHistogramRowWidget(const TSharedRef<STableViewBase>& InOwnerTable);
@@ -109,7 +108,7 @@ private:
 	FLinearColor NodeColor; 
 	float NormalizedValue;
 	float InspectedValue;
-	FText HoverText;
+	FString HoverString;
 };
 
 
@@ -129,11 +128,11 @@ public:
 	void Construct(const FArguments& InArgs);
 
 	TSharedRef<ITableRow> MakeHistogramRowWidget(FGeometryCollectionHistogramItemPtr InInfo, const TSharedRef<STableViewBase>& OwnerTable);
-	void SetComponents(const TArray<UGeometryCollectionComponent*>& InNewComponents);
+	void SetComponents(const TArray<UGeometryCollectionComponent*>& InNewComponents, int32 LevelView);
 	void SetBoneSelection(UGeometryCollectionComponent* RootComponent, const TArray<int32>& InSelection, bool bClearCurrentSelection);
 	void InspectAttribute(EInspectedAttributeEnum InspectedAttribute);
 	void RefreshView(bool bSorted);
-	void RegenerateNodes();
+	void RegenerateNodes(int32 LevelView);
 	void ClearSelection() { ListView->ClearSelection(); }
 	bool IsSelected() { return ListView->GetNumItemsSelected() > 0; }
 
diff --git a/Engine/Plugins/Experimental/ChaosEditor/Source/FractureEditor/Public/FractureEditorModeToolkit.h b/Engine/Plugins/Experimental/ChaosEditor/Source/FractureEditor/Public/FractureEditorModeToolkit.h
--- a/Engine/Plugins/Experimental/ChaosEditor/Source/FractureEditor/Public/FractureEditorModeToolkit.h
+++ b/Engine/Plugins/Experimental/ChaosEditor/Source/FractureEditor/Public/FractureEditorModeToolkit.h
@@ -152,6 +152,7 @@ public:
 
 	void UpdateExplodedVectors(UGeometryCollectionComponent* GeometryCollectionComponent) const;
 
+	void RegenerateOutliner();
 	void RegenerateHistogram();
 
 	TSharedPtr<SWidget> ExplodedViewWidget;
diff --git a/Engine/Plugins/Experimental/ChaosEditor/Source/FractureEditor/Public/FractureTool.cpp b/Engine/Plugins/Experimental/ChaosEditor/Source/FractureEditor/Public/FractureTool.cpp
--- a/Engine/Plugins/Experimental/ChaosEditor/Source/FractureEditor/Public/FractureTool.cpp
+++ b/Engine/Plugins/Experimental/ChaosEditor/Source/FractureEditor/Public/FractureTool.cpp
@@ -134,11 +134,13 @@ void UFractureActionTool::Refresh(FFractureToolContext& Context, FFractureEditor
 	
 	FGeometryCollectionClusteringUtility::UpdateHierarchyLevelOfChildren(GeometryCollectionPtr.Get(), -1);
 
+	Toolkit->RegenerateOutliner();
+	Toolkit->RegenerateHistogram();
+
 	FScopedColorEdit EditBoneColor(GeometryCollectionComponent, true);
 	EditBoneColor.SetSelectedBones(Context.GetSelection());
 
 	FFractureSelectionTools::ToggleSelectedBones(GeometryCollectionComponent, Context.GetSelection(), true);
-	Toolkit->SetBoneSelection(GeometryCollectionComponent, Context.GetSelection(), true);
 
 	Toolkit->UpdateExplodedVectors(GeometryCollectionComponent);
 
diff --git a/Engine/Source/Runtime/Experimental/Chaos/Private/GeometryCollection/GeometryCollectionClusteringUtility.cpp b/Engine/Source/Runtime/Experimental/Chaos/Private/GeometryCollection/GeometryCollectionClusteringUtility.cpp
--- a/Engine/Source/Runtime/Experimental/Chaos/Private/GeometryCollection/GeometryCollectionClusteringUtility.cpp
+++ b/Engine/Source/Runtime/Experimental/Chaos/Private/GeometryCollection/GeometryCollectionClusteringUtility.cpp
@@ -754,8 +754,9 @@ void FGeometryCollectionClusteringUtility::GetLeafBones(FGeometryCollection* Geo
 	}
 
 	const TManagedArray<TSet<int32>>& Children = GeometryCollection->Children;
+	const TManagedArray<int32>& SimulationType = GeometryCollection->SimulationType;
 
-	if (Children[BoneIndex].Num() > 0)
+	if (SimulationType[BoneIndex] == FGeometryCollection::ESimulationTypes::FST_Clustered && Children[BoneIndex].Num() > 0)
 	{
 		for (int32 ChildElement : Children[BoneIndex])
 		{
@@ -764,7 +765,10 @@ void FGeometryCollectionClusteringUtility::GetLeafBones(FGeometryCollection* Geo
 	}
 	else
 	{
-		LeafBonesOut.Push(BoneIndex);
+		if (SimulationType[BoneIndex] == FGeometryCollection::ESimulationTypes::FST_Rigid)
+		{
+			LeafBonesOut.Push(BoneIndex);
+		}	
 	}
 
 }
