Fracture Editor
   * Clustering Utilities GetLeafNode may now specify if the clien tis specifically after rigid leaf nodes.
   * Embedded Geometry tools moved to their own palette.
   * Delete Embedded Geometry tool implemented.

diff --git a/Engine/Plugins/Experimental/ChaosEditor/Source/FractureEditor/Private/FractureEditorCommands.h b/Engine/Plugins/Experimental/ChaosEditor/Source/FractureEditor/Private/FractureEditorCommands.h
--- a/Engine/Plugins/Experimental/ChaosEditor/Source/FractureEditor/Private/FractureEditorCommands.h
+++ b/Engine/Plugins/Experimental/ChaosEditor/Source/FractureEditor/Private/FractureEditorCommands.h
@@ -38,7 +38,10 @@ class FFractureEditorCommands : public TCommands<FFractureEditorCommands>
 		// Generate Commands
 		TSharedPtr< FUICommandInfo > GenerateAsset;
 		TSharedPtr< FUICommandInfo > ResetAsset;
+
+		// Embed Commands
 		TSharedPtr< FUICommandInfo > AddEmbeddedGeometry;
+		TSharedPtr< FUICommandInfo > DeleteEmbeddedGeometry;
 		
 		// Fracture Commands
 		TSharedPtr< FUICommandInfo > Uniform;
diff --git a/Engine/Plugins/Experimental/ChaosEditor/Source/FractureEditor/Private/FractureEditorModeToolkit.cpp b/Engine/Plugins/Experimental/ChaosEditor/Source/FractureEditor/Private/FractureEditorModeToolkit.cpp
--- a/Engine/Plugins/Experimental/ChaosEditor/Source/FractureEditor/Private/FractureEditorModeToolkit.cpp
+++ b/Engine/Plugins/Experimental/ChaosEditor/Source/FractureEditor/Private/FractureEditorModeToolkit.cpp
@@ -453,7 +453,7 @@ void FFractureEditorModeToolkit::OnObjectPostEditChange( UObject* Object, FPrope
 	}
 }
 
-const TArray<FName> FFractureEditorModeToolkit::PaletteNames = { FName(TEXT("Generate")), FName(TEXT("Select")), FName(TEXT("Fracture")), FName(TEXT("Cluster")), FName(TEXT("Properties")) };
+const TArray<FName> FFractureEditorModeToolkit::PaletteNames = { FName(TEXT("Generate")), FName(TEXT("Select")), FName(TEXT("Fracture")), FName(TEXT("Cluster")), FName(TEXT("Embed")), FName(TEXT("Properties")) };
 
 FText FFractureEditorModeToolkit::GetToolPaletteDisplayName(FName Palette) const
 { 
@@ -481,7 +481,6 @@ void FFractureEditorModeToolkit::BuildToolPalette(FName PaletteIndex, class FToo
 	{
 		ToolbarBuilder.AddToolBarButton(Commands.GenerateAsset);
 		ToolbarBuilder.AddToolBarButton(Commands.ResetAsset);
-		ToolbarBuilder.AddToolBarButton(Commands.AddEmbeddedGeometry);
 	}
 	else if (PaletteIndex == TEXT("Select"))
 	{
@@ -511,6 +510,11 @@ void FFractureEditorModeToolkit::BuildToolPalette(FName PaletteIndex, class FToo
 		ToolbarBuilder.AddToolBarButton(Commands.Uncluster);
 		ToolbarBuilder.AddToolBarButton(Commands.MoveUp);
 	}
+	else if (PaletteIndex == TEXT("Embed"))
+	{
+		ToolbarBuilder.AddToolBarButton(Commands.AddEmbeddedGeometry);
+		ToolbarBuilder.AddToolBarButton(Commands.DeleteEmbeddedGeometry);
+	}
 	else if (PaletteIndex == TEXT("Properties"))
 	{
 		ToolbarBuilder.AddToolBarButton(Commands.SetInitialDynamicState);
diff --git a/Engine/Plugins/Experimental/ChaosEditor/Source/FractureEditor/Private/FractureToolClustering.cpp b/Engine/Plugins/Experimental/ChaosEditor/Source/FractureEditor/Private/FractureToolClustering.cpp
--- a/Engine/Plugins/Experimental/ChaosEditor/Source/FractureEditor/Private/FractureToolClustering.cpp
+++ b/Engine/Plugins/Experimental/ChaosEditor/Source/FractureEditor/Private/FractureToolClustering.cpp
@@ -49,7 +49,7 @@ void UFractureToolFlattenAll::Execute(TWeakPtr<FFractureEditorModeToolkit> InToo
 			for (int32 ClusterIndex : Context.GetSelection())
 			{
 				TArray<int32> LeafBones;
-				FGeometryCollectionClusteringUtility::GetLeafBones(Context.GetGeometryCollection().Get(), ClusterIndex, LeafBones);
+				FGeometryCollectionClusteringUtility::GetLeafBones(Context.GetGeometryCollection().Get(), ClusterIndex, true, LeafBones);
 				FGeometryCollectionClusteringUtility::ClusterBonesUnderExistingNode(Context.GetGeometryCollection().Get(), ClusterIndex, LeafBones);
 
 				// Cleanup: Remove any clusters remaining in the flattened branch.
diff --git a/Engine/Plugins/Experimental/ChaosEditor/Source/FractureEditor/Private/FractureToolEmbed.cpp b/Engine/Plugins/Experimental/ChaosEditor/Source/FractureEditor/Private/FractureToolEmbed.cpp
new file mode 100644
--- /dev/null
+++ b/Engine/Plugins/Experimental/ChaosEditor/Source/FractureEditor/Private/FractureToolEmbed.cpp
@@ -0,0 +1,239 @@
+// Copyright Epic Games, Inc. All Rights Reserved.
+
+#include "FractureToolEmbed.h"
+
+#include "Editor.h"
+#include "ScopedTransaction.h"
+#include "Engine/Selection.h"
+#include "Engine/StaticMesh.h"
+
+#include "GeometryCollection/GeometryCollectionComponent.h"
+#include "GeometryCollection/GeometryCollection.h"
+#include "GeometryCollection/GeometryCollectionClusteringUtility.h"
+#include "GeometryCollection/GeometryCollectionConversion.h"
+#include "FractureToolContext.h"
+
+
+#define LOCTEXT_NAMESPACE "FractureToolEmbed"
+
+
+
+FText UFractureToolAddEmbeddedGeometry::GetDisplayText() const
+{
+	return FText(NSLOCTEXT("Fracture", "FractureToolAddEmbeddedGeometry", "Embed"));
+}
+
+FText UFractureToolAddEmbeddedGeometry::GetTooltipText() const
+{
+	return FText(NSLOCTEXT("Fracture", "FractureToolAddEmbeddedGeometryTooltip", "Embed selected static mesh as passive geometry parented to nearest bone. Will be lost if GeometryCollection is Reset!"));
+}
+
+FSlateIcon UFractureToolAddEmbeddedGeometry::GetToolIcon() const
+{
+	return FSlateIcon("FractureEditorStyle", "FractureEditor.AddEmbeddedGeometry");
+}
+
+void UFractureToolAddEmbeddedGeometry::RegisterUICommand(FFractureEditorCommands* BindingContext)
+{
+	UI_COMMAND_EXT(BindingContext, UICommandInfo, "AddEmbeddedGeometry", "Embed", "Embed", EUserInterfaceActionType::ToggleButton, FInputChord());
+	BindingContext->AddEmbeddedGeometry = UICommandInfo;
+}
+
+bool UFractureToolAddEmbeddedGeometry::CanExecute() const
+{
+	return (IsStaticMeshSelected() && IsGeometryCollectionSelected());
+}
+
+void UFractureToolAddEmbeddedGeometry::Execute(TWeakPtr<FFractureEditorModeToolkit> InToolkit)
+{
+	if (!InToolkit.IsValid())
+	{
+		return;
+	}
+
+	// Get selected static meshes
+	TArray<UStaticMeshComponent*> SelectedStaticMeshComponents = GetSelectedStaticMeshComponents();
+
+	TArray<FFractureToolContext> Contexts = GetFractureToolContexts();
+
+	for (FFractureToolContext& Context : Contexts)
+	{
+		FGeometryCollectionEdit GeometryCollectionEdit = Context.GetGeometryCollectionComponent()->EditRestCollection(GeometryCollection::EEditUpdate::RestPhysicsDynamic);
+		UGeometryCollection* FracturedGeometryCollection = GeometryCollectionEdit.GetRestCollection();
+
+		FFractureToolContext::FGeometryCollectionPtr GeometryCollection = Context.GetGeometryCollection();
+		TManagedArray<int32>& Parent = GeometryCollection->Parent;
+		TManagedArray<int32>& SimType = GeometryCollection->SimulationType;
+		TManagedArray<TSet<int32>>& Children = GeometryCollection->Children;
+		TManagedArray<bool>& CollectionSimulatableParticles =
+			GeometryCollection->GetAttribute<bool>(
+				FGeometryCollection::SimulatableParticlesAttribute, FTransformCollection::TransformGroup);
+
+		const FTransform TargetActorTransform(Context.GetGeometryCollectionComponent()->GetOwner()->GetTransform());
+
+		Context.ConvertSelectionToRigidNodes();
+		const TArray<int32>& SelectedBones = Context.GetSelection();
+		for (const int32 SelectedBone : SelectedBones)
+		{
+			for (UStaticMeshComponent* SelectedStaticMeshComponent : SelectedStaticMeshComponents)
+			{
+				const AActor* Actor = SelectedStaticMeshComponent->GetOwner();
+				check(Actor)
+
+					const FTransform SMActorTransform(Actor->GetTransform());
+
+				UStaticMesh* ComponentStaticMesh = SelectedStaticMeshComponent->GetStaticMesh();
+				if (ComponentStaticMesh)
+				{
+					// If any of the static meshes have Nanite enabled, also enable on the new geometry collection asset for convenience.
+					FracturedGeometryCollection->EnableNanite |= ComponentStaticMesh->NaniteSettings.bEnabled;
+				}
+
+				FTransform ComponentTransform = SMActorTransform.GetRelativeTransform(TargetActorTransform);
+				decltype(FGeometryCollectionSource::SourceMaterial) SourceMaterials(SelectedStaticMeshComponent->GetMaterials());
+
+				// Add static mesh geometry to context's geometry collection.
+				FGeometryCollectionConversion::AppendStaticMesh(ComponentStaticMesh, SourceMaterials, ComponentTransform, FracturedGeometryCollection, true);
+
+				int32 NewTransformIndex = Parent.Num() - 1;
+				check(NewTransformIndex > -1)
+
+					// Reparent the new transform to the selected bone
+					Parent[NewTransformIndex] = SelectedBone;
+				Children[SelectedBone].Add(NewTransformIndex);
+
+				// Set properties appropriate to embedded geometry
+				SimType[NewTransformIndex] = FGeometryCollection::ESimulationTypes::FST_None;
+				CollectionSimulatableParticles[NewTransformIndex] = false;
+			}
+		}
+
+		FracturedGeometryCollection->InitializeMaterials();
+
+		FGeometryCollectionClusteringUtility::UpdateHierarchyLevelOfChildren(GeometryCollection.Get(), -1);
+		Context.GetGeometryCollectionComponent()->MarkRenderStateDirty();
+
+		FracturedGeometryCollection->MarkPackageDirty();
+	}
+
+	SetOutlinerComponents(Contexts, InToolkit.Pin().Get());
+}
+
+TArray<UStaticMeshComponent*> UFractureToolAddEmbeddedGeometry::GetSelectedStaticMeshComponents()
+{
+	TArray<UStaticMeshComponent*> SelectedStaticMeshComponents;
+
+	USelection* SelectedActors = GEditor->GetSelectedActors();
+	for (FSelectionIterator Iter(*SelectedActors); Iter; ++Iter)
+	{
+		AActor* Actor = Cast<AActor>(*Iter);
+		if (Actor)
+		{
+			// We don't want to include any static meshes owned by a geometry collection.
+			// (This might be an ISMC or swap out static geometry.)
+			if (Actor->FindComponentByClass<UGeometryCollectionComponent>())
+			{
+				continue;
+			}
+
+			TArray<UStaticMeshComponent*> StaticMeshComponents;
+			Actor->GetComponents<UStaticMeshComponent>(StaticMeshComponents, true);
+			SelectedStaticMeshComponents.Append(StaticMeshComponents);
+		}
+	}
+
+	return SelectedStaticMeshComponents;
+}
+
+
+
+FText UFractureToolDeleteEmbeddedGeometry::GetDisplayText() const
+{
+	return FText(NSLOCTEXT("Fracture", "FractureToolDeleteEmbeddedGeometry", "Delete"));
+}
+
+FText UFractureToolDeleteEmbeddedGeometry::GetTooltipText() const
+{
+	return FText(NSLOCTEXT("Fracture", "FractureToolDeleteEmbeddedGeometryTooltip", "Delete selected embedded geometry nodes. If a cluster or rigid node is selected, all child embedded geometry nodes are deleted."));
+}
+
+FSlateIcon UFractureToolDeleteEmbeddedGeometry::GetToolIcon() const
+{
+	return FSlateIcon("FractureEditorStyle", "FractureEditor.DeleteEmbeddedGeometry");
+}
+
+void UFractureToolDeleteEmbeddedGeometry::RegisterUICommand(FFractureEditorCommands* BindingContext)
+{
+	UI_COMMAND_EXT(BindingContext, UICommandInfo, "DeleteEmbeddedGeometry", "Delete", "Delete", EUserInterfaceActionType::ToggleButton, FInputChord());
+	BindingContext->DeleteEmbeddedGeometry = UICommandInfo;
+}
+
+bool UFractureToolDeleteEmbeddedGeometry::CanExecute() const
+{
+	return IsGeometryCollectionSelected();
+}
+
+void UFractureToolDeleteEmbeddedGeometry::Execute(TWeakPtr<FFractureEditorModeToolkit> InToolkit)
+{
+	if (!InToolkit.IsValid())
+	{
+		return;
+	}
+
+	// Get selected static meshes
+	TArray<FFractureToolContext> Contexts = GetFractureToolContexts();
+
+	for (FFractureToolContext& Context : Contexts)
+	{
+		Context.Sanitize();
+		
+		FGeometryCollectionEdit GeometryCollectionEdit = Context.GetGeometryCollectionComponent()->EditRestCollection(GeometryCollection::EEditUpdate::RestPhysicsDynamic);
+		UGeometryCollection* FracturedGeometryCollection = GeometryCollectionEdit.GetRestCollection();
+
+		FFractureToolContext::FGeometryCollectionPtr GeometryCollection = Context.GetGeometryCollection();
+		const TManagedArray<int32>& SimType = GeometryCollection->SimulationType;
+
+		TArray<int32> EmbeddedGeometryToBeRemoved;
+		const TArray<int32>& SelectedBones = Context.GetSelection();
+		for (const int32 SelectedBone : SelectedBones)
+		{
+			if (SimType[SelectedBone] != FGeometryCollection::ESimulationTypes::FST_None)
+			{
+				// Select all Embedded Geometry found in the selected branch.
+				TArray<int32> LeafBones;
+				FGeometryCollectionClusteringUtility::GetLeafBones(GeometryCollection.Get(), SelectedBone, false, LeafBones);
+				for (int32 LeafBone : LeafBones)
+				{
+					if (SimType[LeafBone] == FGeometryCollection::ESimulationTypes::FST_None)
+					{
+						EmbeddedGeometryToBeRemoved.Add(LeafBone);
+					}
+				}
+			}
+			else
+			{
+				// Selected bone is embedded geometry. Only delete this.
+				EmbeddedGeometryToBeRemoved.Add(SelectedBone);
+			}
+		}
+
+		EmbeddedGeometryToBeRemoved.Sort();
+		GeometryCollection->RemoveElements(FGeometryCollection::TransformGroup, EmbeddedGeometryToBeRemoved);
+
+		Context.GetGeometryCollectionComponent()->MarkRenderStateDirty();
+		FracturedGeometryCollection->MarkPackageDirty();
+	}
+
+	SetOutlinerComponents(Contexts, InToolkit.Pin().Get());
+}
+
+
+
+
+
+
+
+
+
+
+#undef LOCTEXT_NAMESPACE
\ No newline at end of file
diff --git a/Engine/Plugins/Experimental/ChaosEditor/Source/FractureEditor/Private/FractureToolEmbed.h b/Engine/Plugins/Experimental/ChaosEditor/Source/FractureEditor/Private/FractureToolEmbed.h
new file mode 100644
--- /dev/null
+++ b/Engine/Plugins/Experimental/ChaosEditor/Source/FractureEditor/Private/FractureToolEmbed.h
@@ -0,0 +1,47 @@
+// Copyright Epic Games, Inc. All Rights Reserved.
+
+#pragma once
+
+#include "FractureTool.h"
+
+#include "FractureToolEmbed.generated.h"
+
+
+UCLASS()
+class UFractureToolAddEmbeddedGeometry : public UFractureActionTool
+{
+public:
+	GENERATED_BODY()
+
+		UFractureToolAddEmbeddedGeometry(const FObjectInitializer& ObjInit) : Super(ObjInit) {}
+
+	// UFractureActionTool Interface
+	virtual FText GetDisplayText() const override;
+	virtual FText GetTooltipText() const override;
+	virtual FSlateIcon GetToolIcon() const override;
+	virtual void RegisterUICommand(FFractureEditorCommands* BindingContext) override;
+	virtual void Execute(TWeakPtr<FFractureEditorModeToolkit> InToolkit) override;
+	virtual bool CanExecute() const override;
+
+private:
+	static TArray<UStaticMeshComponent*>  GetSelectedStaticMeshComponents();
+};
+
+
+UCLASS()
+class UFractureToolDeleteEmbeddedGeometry : public UFractureActionTool
+{
+public:
+	GENERATED_BODY()
+
+	UFractureToolDeleteEmbeddedGeometry(const FObjectInitializer& ObjInit) : Super(ObjInit) {}
+
+	// UFractureActionTool Interface
+	virtual FText GetDisplayText() const override;
+	virtual FText GetTooltipText() const override;
+	virtual FSlateIcon GetToolIcon() const override;
+	virtual void RegisterUICommand(FFractureEditorCommands* BindingContext) override;
+	virtual void Execute(TWeakPtr<FFractureEditorModeToolkit> InToolkit) override;
+	virtual bool CanExecute() const override;
+};
+
diff --git a/Engine/Plugins/Experimental/ChaosEditor/Source/FractureEditor/Private/FractureToolGenerators.cpp b/Engine/Plugins/Experimental/ChaosEditor/Source/FractureEditor/Private/FractureToolGenerators.cpp
--- a/Engine/Plugins/Experimental/ChaosEditor/Source/FractureEditor/Private/FractureToolGenerators.cpp
+++ b/Engine/Plugins/Experimental/ChaosEditor/Source/FractureEditor/Private/FractureToolGenerators.cpp
@@ -377,133 +377,4 @@ void UFractureToolResetAsset::Execute(TWeakPtr<FFractureEditorModeToolkit> InToo
 }
 
 
-
-
-FText UFractureToolAddEmbeddedGeometry::GetDisplayText() const
-{
-	return FText(NSLOCTEXT("Fracture", "FractureToolAddEmbeddedGeometry", "Embed"));
-}
-
-FText UFractureToolAddEmbeddedGeometry::GetTooltipText() const
-{
-	return FText(NSLOCTEXT("Fracture", "FractureToolAddEmbeddedGeometryTooltip", "Embed selected static mesh as passive geometry parented to nearest bone. Will be lost if GeometryCollection is Reset!"));
-}
-
-FSlateIcon UFractureToolAddEmbeddedGeometry::GetToolIcon() const
-{
-	return FSlateIcon("FractureEditorStyle", "FractureEditor.AddEmbeddedGeometry");
-}
-
-void UFractureToolAddEmbeddedGeometry::RegisterUICommand(FFractureEditorCommands* BindingContext)
-{
-	UI_COMMAND_EXT(BindingContext, UICommandInfo, "AddEmbeddedGeometry", "Embed", "Embed", EUserInterfaceActionType::ToggleButton, FInputChord());
-	BindingContext->AddEmbeddedGeometry = UICommandInfo;
-}
-
-bool UFractureToolAddEmbeddedGeometry::CanExecute() const
-{
-	return (IsStaticMeshSelected() && IsGeometryCollectionSelected());
-}
-
-void UFractureToolAddEmbeddedGeometry::Execute(TWeakPtr<FFractureEditorModeToolkit> InToolkit)
-{
-	if (!InToolkit.IsValid())
-	{
-		return;
-	}
-
-	// Get selected static meshes
-	TArray<UStaticMeshComponent*> SelectedStaticMeshComponents = GetSelectedStaticMeshComponents();
-
-	TArray<FFractureToolContext> Contexts = GetFractureToolContexts();
-
-	for (FFractureToolContext& Context : Contexts)
-	{
-		FGeometryCollectionEdit GeometryCollectionEdit = Context.GetGeometryCollectionComponent()->EditRestCollection(GeometryCollection::EEditUpdate::RestPhysicsDynamic);
-		UGeometryCollection* FracturedGeometryCollection = GeometryCollectionEdit.GetRestCollection();
-
-		FFractureToolContext::FGeometryCollectionPtr GeometryCollection = Context.GetGeometryCollection();
-		TManagedArray<int32>& Parent = GeometryCollection->Parent;
-		TManagedArray<int32>& SimType = GeometryCollection->SimulationType;
-		TManagedArray<TSet<int32>>& Children = GeometryCollection->Children;
-		TManagedArray<bool>& CollectionSimulatableParticles = 
-			GeometryCollection->GetAttribute<bool>(
-				FGeometryCollection::SimulatableParticlesAttribute, FTransformCollection::TransformGroup);
-
-		const FTransform TargetActorTransform(Context.GetGeometryCollectionComponent()->GetOwner()->GetTransform());
-		
-		Context.ConvertSelectionToRigidNodes();
-		const TArray<int32>& SelectedBones = Context.GetSelection();
-		for (const int32 SelectedBone : SelectedBones)
-		{
-			for (UStaticMeshComponent* SelectedStaticMeshComponent : SelectedStaticMeshComponents)
-			{
-				const AActor* Actor = SelectedStaticMeshComponent->GetOwner();
-				check(Actor)
-
-				const FTransform SMActorTransform(Actor->GetTransform());
-
-				UStaticMesh* ComponentStaticMesh = SelectedStaticMeshComponent->GetStaticMesh();
-				if (ComponentStaticMesh)
-				{
-					// If any of the static meshes have Nanite enabled, also enable on the new geometry collection asset for convenience.
-					FracturedGeometryCollection->EnableNanite |= ComponentStaticMesh->NaniteSettings.bEnabled;
-				}
-
-				FTransform ComponentTransform = SMActorTransform.GetRelativeTransform(TargetActorTransform);
-				decltype(FGeometryCollectionSource::SourceMaterial) SourceMaterials(SelectedStaticMeshComponent->GetMaterials());
-
-				// Add static mesh geometry to context's geometry collection.
-				FGeometryCollectionConversion::AppendStaticMesh(ComponentStaticMesh, SourceMaterials, ComponentTransform, FracturedGeometryCollection, true);
-
-				int32 NewTransformIndex = Parent.Num() - 1;
-				check(NewTransformIndex > -1)
-
-				// Reparent the new transform to the selected bone
-				Parent[NewTransformIndex] = SelectedBone;
-				Children[SelectedBone].Add(NewTransformIndex);
-
-				// Set properties appropriate to embedded geometry
-				SimType[NewTransformIndex] = FGeometryCollection::ESimulationTypes::FST_None;
-				CollectionSimulatableParticles[NewTransformIndex] = false;
-			}	
-		}
-
-		FracturedGeometryCollection->InitializeMaterials();
-
-		FGeometryCollectionClusteringUtility::UpdateHierarchyLevelOfChildren(GeometryCollection.Get(), -1);
-		Context.GetGeometryCollectionComponent()->MarkRenderStateDirty();
-		
-		FracturedGeometryCollection->MarkPackageDirty();
-	}
-
-	SetOutlinerComponents(Contexts, InToolkit.Pin().Get());
-}
-
-TArray<UStaticMeshComponent*> UFractureToolAddEmbeddedGeometry::GetSelectedStaticMeshComponents()
-{
-	TArray<UStaticMeshComponent*> SelectedStaticMeshComponents;
-
-	USelection* SelectedActors = GEditor->GetSelectedActors();
-	for (FSelectionIterator Iter(*SelectedActors); Iter; ++Iter)
-	{
-		AActor* Actor = Cast<AActor>(*Iter);
-		if (Actor)
-		{
-			// We don't want to include any static meshes owned by a geometry collection.
-			// (This might be an ISMC or swap out static geometry.)
-			if (Actor->FindComponentByClass<UGeometryCollectionComponent>())
-			{
-				continue;
-			}
-			
-			TArray<UStaticMeshComponent*> StaticMeshComponents;
-			Actor->GetComponents<UStaticMeshComponent>(StaticMeshComponents, true);
-			SelectedStaticMeshComponents.Append(StaticMeshComponents);
-		}
-	}
-
-	return SelectedStaticMeshComponents;
-}
-
 #undef LOCTEXT_NAMESPACE
diff --git a/Engine/Plugins/Experimental/ChaosEditor/Source/FractureEditor/Private/FractureToolGenerators.h b/Engine/Plugins/Experimental/ChaosEditor/Source/FractureEditor/Private/FractureToolGenerators.h
--- a/Engine/Plugins/Experimental/ChaosEditor/Source/FractureEditor/Private/FractureToolGenerators.h
+++ b/Engine/Plugins/Experimental/ChaosEditor/Source/FractureEditor/Private/FractureToolGenerators.h
@@ -62,23 +62,3 @@ public:
 };
 
 
-UCLASS()
-class UFractureToolAddEmbeddedGeometry : public UFractureActionTool
-{
-public:
-	GENERATED_BODY()
-
-	UFractureToolAddEmbeddedGeometry(const FObjectInitializer& ObjInit) : Super(ObjInit) {}
-
-	// UFractureActionTool Interface
-	virtual FText GetDisplayText() const override;
-	virtual FText GetTooltipText() const override;
-	virtual FSlateIcon GetToolIcon() const override;
-	virtual void RegisterUICommand(FFractureEditorCommands* BindingContext) override;
-	virtual void Execute(TWeakPtr<FFractureEditorModeToolkit> InToolkit) override;
-	virtual bool CanExecute() const override;
-
-private:
-	static TArray<UStaticMeshComponent*>  GetSelectedStaticMeshComponents();
-};
-
diff --git a/Engine/Source/Runtime/Experimental/Chaos/Private/GeometryCollection/GeometryCollectionClusteringUtility.cpp b/Engine/Source/Runtime/Experimental/Chaos/Private/GeometryCollection/GeometryCollectionClusteringUtility.cpp
--- a/Engine/Source/Runtime/Experimental/Chaos/Private/GeometryCollection/GeometryCollectionClusteringUtility.cpp
+++ b/Engine/Source/Runtime/Experimental/Chaos/Private/GeometryCollection/GeometryCollectionClusteringUtility.cpp
@@ -746,7 +746,7 @@ void FGeometryCollectionClusteringUtility::ContextBasedClusterSelection(
 
 }
 
-void FGeometryCollectionClusteringUtility::GetLeafBones(FGeometryCollection* GeometryCollection, int BoneIndex, TArray<int32>& LeafBonesOut)
+void FGeometryCollectionClusteringUtility::GetLeafBones(FGeometryCollection* GeometryCollection, int BoneIndex, bool bOnlyRigids, TArray<int32>& LeafBonesOut)
 {
 	if (!ensure(BoneIndex >= 0))
 	{
@@ -756,19 +756,20 @@ void FGeometryCollectionClusteringUtility::GetLeafBones(FGeometryCollection* Geo
 	const TManagedArray<TSet<int32>>& Children = GeometryCollection->Children;
 	const TManagedArray<int32>& SimulationType = GeometryCollection->SimulationType;
 
-	if (SimulationType[BoneIndex] == FGeometryCollection::ESimulationTypes::FST_Clustered && Children[BoneIndex].Num() > 0)
+	if (!bOnlyRigids && Children[BoneIndex].Num() == 0)
 	{
-		for (int32 ChildElement : Children[BoneIndex])
-		{
-			GetLeafBones(GeometryCollection, ChildElement, LeafBonesOut);
-		}
+		LeafBonesOut.Push(BoneIndex);
 	}
-	else
+	else if (bOnlyRigids && GeometryCollection->IsRigid(BoneIndex))
 	{
-		if (SimulationType[BoneIndex] == FGeometryCollection::ESimulationTypes::FST_Rigid)
+		LeafBonesOut.Push(BoneIndex);
+	}
+	else if (Children[BoneIndex].Num() > 0)
+	{
+		for (int32 ChildElement : Children[BoneIndex])
 		{
-			LeafBonesOut.Push(BoneIndex);
-		}	
+			GetLeafBones(GeometryCollection, ChildElement, bOnlyRigids, LeafBonesOut);
+		}
 	}
 
 }
diff --git a/Engine/Source/Runtime/Experimental/Chaos/Public/GeometryCollection/GeometryCollectionClusteringUtility.h b/Engine/Source/Runtime/Experimental/Chaos/Public/GeometryCollection/GeometryCollectionClusteringUtility.h
--- a/Engine/Source/Runtime/Experimental/Chaos/Public/GeometryCollection/GeometryCollectionClusteringUtility.h
+++ b/Engine/Source/Runtime/Experimental/Chaos/Public/GeometryCollection/GeometryCollectionClusteringUtility.h
@@ -103,8 +103,8 @@ public:
 	/** Make logical editor selections based on the current selected hierarchy view level */
 	static void ContextBasedClusterSelection(FGeometryCollection* GeometryCollection, int ViewLevel, const TArray<int32>& SelectedComponentBonesIn, TArray<int32>& SelectedComponentBonesOut, TArray<int32>& HighlightedComponentBonesOut);
 
-	/** return an array of all child leaf nodes below the specified node */
-	static void GetLeafBones(FGeometryCollection* GeometryCollection, int BoneIndex, TArray<int32>& LeafBonesOut);
+	/** return an array of all child leaf nodes below the specified node. If bOnlyRigids is true, the first Rigid node dound is considered a leaf, regardless of an children it might have. */
+	static void GetLeafBones(FGeometryCollection* GeometryCollection, int BoneIndex, bool bOnlyRigids, TArray<int32>& LeafBonesOut);
 
 	/** move the selected node up a level in direction of root */
 	static void MoveUpOneHierarchyLevel(FGeometryCollection* GeometryCollection, const TArray<int32>& SelectedBones);
diff --git a/Engine/Source/Runtime/Experimental/GeometryCollectionEngine/Private/GeometryCollection/GeometryCollectionComponent.cpp b/Engine/Source/Runtime/Experimental/GeometryCollectionEngine/Private/GeometryCollection/GeometryCollectionComponent.cpp
--- a/Engine/Source/Runtime/Experimental/GeometryCollectionEngine/Private/GeometryCollection/GeometryCollectionComponent.cpp
+++ b/Engine/Source/Runtime/Experimental/GeometryCollectionEngine/Private/GeometryCollection/GeometryCollectionComponent.cpp
@@ -2313,7 +2313,7 @@ void FScopedColorEdit::SelectBones(GeometryCollection::ESelectionMode SelectionM
 			for (int32 RootElement : Roots)
 			{
 				TArray<int32> LeafBones;
-				FGeometryCollectionClusteringUtility::GetLeafBones(GeometryCollectionPtr.Get(), RootElement, LeafBones);
+				FGeometryCollectionClusteringUtility::GetLeafBones(GeometryCollectionPtr.Get(), RootElement, true, LeafBones);
 				AppendSelectedBones(LeafBones);
 			}
 		}
@@ -2327,7 +2327,7 @@ void FScopedColorEdit::SelectBones(GeometryCollection::ESelectionMode SelectionM
 			for (int32 RootElement : Roots)
 			{
 				TArray<int32> LeafBones;
-				FGeometryCollectionClusteringUtility::GetLeafBones(GeometryCollectionPtr.Get(), RootElement, LeafBones);
+				FGeometryCollectionClusteringUtility::GetLeafBones(GeometryCollectionPtr.Get(), RootElement, true, LeafBones);
 
 				for (int32 Element : LeafBones)
 				{
@@ -2409,7 +2409,7 @@ void FScopedColorEdit::SelectBones(GeometryCollection::ESelectionMode SelectionM
 			{
 				int32 ParentBone = Parents[Bone];
 				TArray<int32> LeafBones;
-				FGeometryCollectionClusteringUtility::GetLeafBones(GeometryCollectionPtr.Get(), ParentBone, LeafBones);
+				FGeometryCollectionClusteringUtility::GetLeafBones(GeometryCollectionPtr.Get(), ParentBone, true, LeafBones);
 
 				for (int32 Element : LeafBones)
 				{
