add missing collision settings property for uniform voronoi fracture

diff --git a/Engine/Plugins/Experimental/ChaosEditor/Source/FractureEditor/Private/FractureToolUniform.cpp b/Engine/Plugins/Experimental/ChaosEditor/Source/FractureEditor/Private/FractureToolUniform.cpp
--- a/Engine/Plugins/Experimental/ChaosEditor/Source/FractureEditor/Private/FractureToolUniform.cpp
+++ b/Engine/Plugins/Experimental/ChaosEditor/Source/FractureEditor/Private/FractureToolUniform.cpp
@@ -41,6 +41,7 @@ TArray<UObject*> UFractureToolUniform::GetSettingsObjects() const
 { 
 	TArray<UObject*> AllSettings; 
 	AllSettings.Add(CutterSettings);
+	AllSettings.Add(CollisionSettings);
 	AllSettings.Add(UniformSettings);
 	return AllSettings;
 }
