add missing includes for fracture mesh cut
#rnx

diff --git a/Engine/Plugins/Experimental/ChaosEditor/Source/FractureEditor/Private/FractureToolMeshCut.cpp b/Engine/Plugins/Experimental/ChaosEditor/Source/FractureEditor/Private/FractureToolMeshCut.cpp
--- a/Engine/Plugins/Experimental/ChaosEditor/Source/FractureEditor/Private/FractureToolMeshCut.cpp
+++ b/Engine/Plugins/Experimental/ChaosEditor/Source/FractureEditor/Private/FractureToolMeshCut.cpp
@@ -7,6 +7,8 @@
 #include "GeometryCollection/GeometryCollectionObject.h"
 #include "FractureEditorModeToolkit.h"
 #include "FractureToolContext.h"
+#include "Components/StaticMeshComponent.h"
+#include "Engine/StaticMesh.h"
 
 #define LOCTEXT_NAMESPACE "FractureMesh"
 
