Fracture Tool -- Automatic embedded geometry tool implemented.

diff --git a/Engine/Plugins/Experimental/ChaosEditor/Source/FractureEditor/Private/FractureEditorCommands.h b/Engine/Plugins/Experimental/ChaosEditor/Source/FractureEditor/Private/FractureEditorCommands.h
--- a/Engine/Plugins/Experimental/ChaosEditor/Source/FractureEditor/Private/FractureEditorCommands.h
+++ b/Engine/Plugins/Experimental/ChaosEditor/Source/FractureEditor/Private/FractureEditorCommands.h
@@ -41,6 +41,7 @@ class FFractureEditorCommands : public TCommands<FFractureEditorCommands>
 
 		// Embed Commands
 		TSharedPtr< FUICommandInfo > AddEmbeddedGeometry;
+		TSharedPtr< FUICommandInfo > AutoEmbedGeometry;
 		TSharedPtr< FUICommandInfo > DeleteEmbeddedGeometry;
 		
 		// Fracture Commands
diff --git a/Engine/Plugins/Experimental/ChaosEditor/Source/FractureEditor/Private/FractureEditorModeToolkit.cpp b/Engine/Plugins/Experimental/ChaosEditor/Source/FractureEditor/Private/FractureEditorModeToolkit.cpp
--- a/Engine/Plugins/Experimental/ChaosEditor/Source/FractureEditor/Private/FractureEditorModeToolkit.cpp
+++ b/Engine/Plugins/Experimental/ChaosEditor/Source/FractureEditor/Private/FractureEditorModeToolkit.cpp
@@ -514,6 +514,7 @@ void FFractureEditorModeToolkit::BuildToolPalette(FName PaletteIndex, class FToo
 	else if (PaletteIndex == TEXT("Embed"))
 	{
 		ToolbarBuilder.AddToolBarButton(Commands.AddEmbeddedGeometry);
+		ToolbarBuilder.AddToolBarButton(Commands.AutoEmbedGeometry);
 		ToolbarBuilder.AddToolBarButton(Commands.DeleteEmbeddedGeometry);
 	}
 	else if (PaletteIndex == TEXT("Properties"))
diff --git a/Engine/Plugins/Experimental/ChaosEditor/Source/FractureEditor/Private/FractureToolEmbed.cpp b/Engine/Plugins/Experimental/ChaosEditor/Source/FractureEditor/Private/FractureToolEmbed.cpp
--- a/Engine/Plugins/Experimental/ChaosEditor/Source/FractureEditor/Private/FractureToolEmbed.cpp
+++ b/Engine/Plugins/Experimental/ChaosEditor/Source/FractureEditor/Private/FractureToolEmbed.cpp
@@ -11,6 +11,9 @@
 #include "GeometryCollection/GeometryCollection.h"
 #include "GeometryCollection/GeometryCollectionClusteringUtility.h"
 #include "GeometryCollection/GeometryCollectionConversion.h"
+#include "GeometryCollection/GeometryCollectionConvexUtility.h"
+#include "GeometryCollection/GeometryCollectionAlgo.h"
+#include "Chaos/Convex.h"
 #include "FractureToolContext.h"
 
 
@@ -25,7 +28,7 @@ FText UFractureToolAddEmbeddedGeometry::GetDisplayText() const
 
 FText UFractureToolAddEmbeddedGeometry::GetTooltipText() const
 {
-	return FText(NSLOCTEXT("Fracture", "FractureToolAddEmbeddedGeometryTooltip", "Embed selected static mesh as passive geometry parented to nearest bone. Will be lost if GeometryCollection is Reset!"));
+	return FText(NSLOCTEXT("Fracture", "FractureToolAddEmbeddedGeometryTooltip", "Embed selected static mesh as passive geometry parented to selected bone. Will be lost if GeometryCollection is Reset!"));
 }
 
 FSlateIcon UFractureToolAddEmbeddedGeometry::GetToolIcon() const
@@ -62,6 +65,7 @@ void UFractureToolAddEmbeddedGeometry::Execute(TWeakPtr<FFractureEditorModeToolk
 		UGeometryCollection* FracturedGeometryCollection = GeometryCollectionEdit.GetRestCollection();
 
 		FFractureToolContext::FGeometryCollectionPtr GeometryCollection = Context.GetGeometryCollection();
+		const TManagedArray<FTransform>& Transform = GeometryCollection->Transform;
 		TManagedArray<int32>& Parent = GeometryCollection->Parent;
 		TManagedArray<int32>& SimType = GeometryCollection->SimulationType;
 		TManagedArray<TSet<int32>>& Children = GeometryCollection->Children;
@@ -75,12 +79,14 @@ void UFractureToolAddEmbeddedGeometry::Execute(TWeakPtr<FFractureEditorModeToolk
 		const TArray<int32>& SelectedBones = Context.GetSelection();
 		for (const int32 SelectedBone : SelectedBones)
 		{
+			FTransform BoneGlobalTransform = GeometryCollectionAlgo::GlobalMatrix(Transform, Parent, SelectedBone);
+			
 			for (UStaticMeshComponent* SelectedStaticMeshComponent : SelectedStaticMeshComponents)
 			{
 				const AActor* Actor = SelectedStaticMeshComponent->GetOwner();
 				check(Actor)
 
-					const FTransform SMActorTransform(Actor->GetTransform());
+				const FTransform SMActorTransform(Actor->GetTransform());
 
 				UStaticMesh* ComponentStaticMesh = SelectedStaticMeshComponent->GetStaticMesh();
 				if (ComponentStaticMesh)
@@ -90,9 +96,12 @@ void UFractureToolAddEmbeddedGeometry::Execute(TWeakPtr<FFractureEditorModeToolk
 				}
 
 				FTransform ComponentTransform = SMActorTransform.GetRelativeTransform(TargetActorTransform);
-				
+				FTransform BoneTransform = ComponentTransform.GetRelativeTransform(BoneGlobalTransform);
+
+				decltype(FGeometryCollectionSource::SourceMaterial) SourceMaterials(SelectedStaticMeshComponent->GetMaterials());
+
 				// Add static mesh geometry to context's geometry collection.
-				FGeometryCollectionConversion::AppendStaticMesh(ComponentStaticMesh, SelectedStaticMeshComponent, ComponentTransform, FracturedGeometryCollection, true);
+				FGeometryCollectionConversion::AppendStaticMesh(ComponentStaticMesh, SelectedStaticMeshComponent, BoneTransform, FracturedGeometryCollection, true);
 
 				int32 NewTransformIndex = Parent.Num() - 1;
 				check(NewTransformIndex > -1)
@@ -146,6 +155,183 @@ TArray<UStaticMeshComponent*> UFractureToolAddEmbeddedGeometry::GetSelectedStati
 
 
 
+FText UFractureToolAutoEmbedGeometry::GetDisplayText() const
+{
+	return FText(NSLOCTEXT("Fracture", "FractureToolAutoEmbedGeometry", "Auto"));
+}
+
+FText UFractureToolAutoEmbedGeometry::GetTooltipText() const
+{
+	return FText(NSLOCTEXT("Fracture", "FractureToolAutoEmbedGeometryTooltip", "Embed selected static meshes as passive geometry parented to nearest bone. Will be lost if GeometryCollection is Reset!"));
+}
+
+FSlateIcon UFractureToolAutoEmbedGeometry::GetToolIcon() const
+{
+	return FSlateIcon("FractureEditorStyle", "FractureEditor.AutoEmbedGeometry");
+}
+
+void UFractureToolAutoEmbedGeometry::RegisterUICommand(FFractureEditorCommands* BindingContext)
+{
+	UI_COMMAND_EXT(BindingContext, UICommandInfo, "AutoEmbedGeometry", "Auto", "Auto", EUserInterfaceActionType::ToggleButton, FInputChord());
+	BindingContext->AutoEmbedGeometry = UICommandInfo;
+}
+
+bool UFractureToolAutoEmbedGeometry::CanExecute() const
+{
+	return (IsStaticMeshSelected() && IsGeometryCollectionSelected());
+}
+
+void UFractureToolAutoEmbedGeometry::Execute(TWeakPtr<FFractureEditorModeToolkit> InToolkit)
+{
+	if (!InToolkit.IsValid())
+	{
+		return;
+	}
+
+	// Get selected static meshes
+	TArray<UStaticMeshComponent*> SelectedStaticMeshComponents = GetSelectedStaticMeshComponents();
+	TArray<FFractureToolContext> Contexts = GetFractureToolContexts();
+
+	// For each static mesh component, we iterate all the convex hulls and determine which convex hull best contains the worldspace pivot of the static mesh.
+	for (UStaticMeshComponent* SelectedStaticMeshComponent : SelectedStaticMeshComponents)
+	{
+		// Static Mesh world space location
+		FVector SMLocation = SelectedStaticMeshComponent->GetComponentLocation();
+
+		// Determine the "closest" bone
+		FFractureToolContext::FGeometryCollectionPtr ClosestGeometryCollection;
+		int32 ClosestConvex = INDEX_NONE;
+		Chaos::FReal ClosestPhi = TNumericLimits<float>::Max();
+		UGeometryCollectionComponent* ClosestComponent = nullptr;
+		
+		for (FFractureToolContext& Context : Contexts)
+		{
+			FGeometryCollection* GeometryCollection = Context.GetGeometryCollection().Get();
+			FGeometryCollectionConvexUtility::FGeometryCollectionConvexData ConvexData = FGeometryCollectionConvexUtility::GetValidConvexHullData(GeometryCollection);
+			const TManagedArray<FTransform>& Transform = GeometryCollection->Transform;
+			const TManagedArray<int32>& Parent = GeometryCollection->Parent;
+			TArray<FTransform> BoneGlobalTransforms;
+			GeometryCollectionAlgo::GlobalMatrices(Transform, Parent, BoneGlobalTransforms);
+			
+			FTransform WorldToComponent = Context.GetGeometryCollectionComponent()->GetComponentToWorld().Inverse();
+			FVector ComponentSpaceLocation = WorldToComponent.TransformPosition(SMLocation);
+
+			int32 NumConvex = ConvexData.ConvexHull.Num();
+			for (int32 ConvexIndex = 0; ConvexIndex < NumConvex; ++ConvexIndex)
+			{
+				// transform into bone space where the convex hull is described
+				int32 TransformIndex = ConvexData.TransformToConvexIndex.Find(ConvexIndex);
+				FTransform ComponentToBone = BoneGlobalTransforms[TransformIndex].Inverse();
+				FVector BoneSpaceLocation = ComponentToBone.TransformPosition(ComponentSpaceLocation);
+				
+				Chaos::FVec3 Normal;
+				Chaos::FReal Phi = ConvexData.ConvexHull[ConvexIndex]->PhiWithNormal(BoneSpaceLocation, Normal);
+				UE_LOG(LogTemp, Warning, TEXT(" Cone %d Phi %f"), ConvexData.TransformToConvexIndex.Find(ConvexIndex), Phi);
+				if (Phi < ClosestPhi)
+				{
+					ClosestPhi = Phi;
+					ClosestGeometryCollection = Context.GetGeometryCollection();
+					ClosestConvex = ConvexIndex;
+					ClosestComponent = Context.GetGeometryCollectionComponent();
+				}
+			}
+		}
+
+		if (ClosestGeometryCollection.IsValid())
+		{
+			// Which bone points to the closest convex?
+			const TManagedArray<FTransform>& Transform = ClosestGeometryCollection->Transform;
+			TManagedArray<int32>& Parent = ClosestGeometryCollection->Parent;
+			TManagedArray<int32>& SimType = ClosestGeometryCollection->SimulationType;
+			TManagedArray<TSet<int32>>& Children = ClosestGeometryCollection->Children;
+			TManagedArray<bool>& CollectionSimulatableParticles =
+				ClosestGeometryCollection->GetAttribute<bool>(
+					FGeometryCollection::SimulatableParticlesAttribute, FTransformCollection::TransformGroup);
+			TManagedArray<int32>& TransformToConvexIndex =
+				ClosestGeometryCollection->GetAttribute<int32>("TransformToConvexIndex", FTransformCollection::TransformGroup);
+
+			const int32 BoneIndex = TransformToConvexIndex.Find(ClosestConvex);
+
+			if (BoneIndex > INDEX_NONE)
+			{
+				// We found the closest bone, now we embed the geometry.
+				FGeometryCollectionEdit GeometryCollectionEdit = ClosestComponent->EditRestCollection(GeometryCollection::EEditUpdate::RestPhysicsDynamic);
+				UGeometryCollection* FracturedGeometryCollection = GeometryCollectionEdit.GetRestCollection();
+
+				FTransform BoneGlobalTransform = GeometryCollectionAlgo::GlobalMatrix(Transform, Parent, BoneIndex);
+
+				const FTransform TargetActorTransform(ClosestComponent->GetOwner()->GetTransform());
+				const AActor* Actor = SelectedStaticMeshComponent->GetOwner();
+				check(Actor)
+
+				const FTransform SMActorTransform(Actor->GetTransform());
+				FTransform ComponentTransform = SMActorTransform.GetRelativeTransform(TargetActorTransform);
+				FTransform BoneTransform = ComponentTransform.GetRelativeTransform(BoneGlobalTransform);
+
+				UStaticMesh* ComponentStaticMesh = SelectedStaticMeshComponent->GetStaticMesh();
+				if (ComponentStaticMesh)
+				{
+					// If any of the static meshes have Nanite enabled, also enable on the new geometry collection asset for convenience.
+					FracturedGeometryCollection->EnableNanite |= ComponentStaticMesh->NaniteSettings.bEnabled;
+				}
+
+				// Add static mesh geometry to context's geometry collection.
+				FGeometryCollectionConversion::AppendStaticMesh(ComponentStaticMesh, SelectedStaticMeshComponent, BoneTransform, FracturedGeometryCollection, true);
+
+				// A new transform has been appended to the TransformGroup
+				int32 NewTransformIndex = Parent.Num() - 1;
+				check(NewTransformIndex > -1)
+
+				// Reparent the new transform to the selected bone
+				Parent[NewTransformIndex] = BoneIndex;
+				Children[BoneIndex].Add(NewTransformIndex);
+
+				// Set properties appropriate to embedded geometry
+				SimType[NewTransformIndex] = FGeometryCollection::ESimulationTypes::FST_None;
+				CollectionSimulatableParticles[NewTransformIndex] = false;
+
+				FracturedGeometryCollection->InitializeMaterials();
+
+				FGeometryCollectionClusteringUtility::UpdateHierarchyLevelOfChildren(ClosestGeometryCollection.Get(), -1);
+				ClosestComponent->MarkRenderStateDirty();
+
+				FracturedGeometryCollection->MarkPackageDirty();
+			}
+		}
+	}
+
+	SetOutlinerComponents(Contexts, InToolkit.Pin().Get());
+}
+
+TArray<UStaticMeshComponent*> UFractureToolAutoEmbedGeometry::GetSelectedStaticMeshComponents()
+{
+	TArray<UStaticMeshComponent*> SelectedStaticMeshComponents;
+
+	USelection* SelectedActors = GEditor->GetSelectedActors();
+	for (FSelectionIterator Iter(*SelectedActors); Iter; ++Iter)
+	{
+		AActor* Actor = Cast<AActor>(*Iter);
+		if (Actor)
+		{
+			// We don't want to include any static meshes owned by a geometry collection.
+			// (This might be an ISMC or swap out static geometry.)
+			if (Actor->FindComponentByClass<UGeometryCollectionComponent>())
+			{
+				continue;
+			}
+
+			TArray<UStaticMeshComponent*> StaticMeshComponents;
+			Actor->GetComponents<UStaticMeshComponent>(StaticMeshComponents, true);
+			SelectedStaticMeshComponents.Append(StaticMeshComponents);
+		}
+	}
+
+	return SelectedStaticMeshComponents;
+}
+
+
+
+
 FText UFractureToolDeleteEmbeddedGeometry::GetDisplayText() const
 {
 	return FText(NSLOCTEXT("Fracture", "FractureToolDeleteEmbeddedGeometry", "Delete"));
diff --git a/Engine/Plugins/Experimental/ChaosEditor/Source/FractureEditor/Private/FractureToolEmbed.h b/Engine/Plugins/Experimental/ChaosEditor/Source/FractureEditor/Private/FractureToolEmbed.h
--- a/Engine/Plugins/Experimental/ChaosEditor/Source/FractureEditor/Private/FractureToolEmbed.h
+++ b/Engine/Plugins/Experimental/ChaosEditor/Source/FractureEditor/Private/FractureToolEmbed.h
@@ -28,6 +28,27 @@ private:
 };
 
 
+UCLASS()
+class UFractureToolAutoEmbedGeometry : public UFractureActionTool
+{
+public:
+	GENERATED_BODY()
+
+	UFractureToolAutoEmbedGeometry(const FObjectInitializer& ObjInit) : Super(ObjInit) {}
+
+	// UFractureActionTool Interface
+	virtual FText GetDisplayText() const override;
+	virtual FText GetTooltipText() const override;
+	virtual FSlateIcon GetToolIcon() const override;
+	virtual void RegisterUICommand(FFractureEditorCommands* BindingContext) override;
+	virtual void Execute(TWeakPtr<FFractureEditorModeToolkit> InToolkit) override;
+	virtual bool CanExecute() const override;
+
+private:
+	static TArray<UStaticMeshComponent*>  GetSelectedStaticMeshComponents();
+};
+
+
 UCLASS()
 class UFractureToolDeleteEmbeddedGeometry : public UFractureActionTool
 {
