Fracture Tool -- Histogram has options to view clusters, rigids, embedded geometry.

diff --git a/Engine/Plugins/Experimental/ChaosEditor/Source/FractureEditor/Private/FractureEditorModeToolkit.cpp b/Engine/Plugins/Experimental/ChaosEditor/Source/FractureEditor/Private/FractureEditorModeToolkit.cpp
--- a/Engine/Plugins/Experimental/ChaosEditor/Source/FractureEditor/Private/FractureEditorModeToolkit.cpp
+++ b/Engine/Plugins/Experimental/ChaosEditor/Source/FractureEditor/Private/FractureEditorModeToolkit.cpp
@@ -444,6 +444,12 @@ void FFractureEditorModeToolkit::OnObjectPostEditChange( UObject* Object, FPrope
 			UHistogramSettings* HistogramSettings = GetMutableDefault<UHistogramSettings>();
 			HistogramView->InspectAttribute(HistogramSettings->InspectedAttribute);
 		}
+		else if ( (PropertyChangedEvent.Property->GetFName() == GET_MEMBER_NAME_CHECKED(UHistogramSettings, bShowRigids)) ||
+			(PropertyChangedEvent.Property->GetFName() == GET_MEMBER_NAME_CHECKED(UHistogramSettings, bShowClusters)) ||
+			(PropertyChangedEvent.Property->GetFName() == GET_MEMBER_NAME_CHECKED(UHistogramSettings, bShowEmbedded)) )
+		{
+			HistogramView->RegenerateNodes(GetLevelViewValue());
+		}
 		else if (PropertyChangedEvent.Property->GetFName() == GET_MEMBER_NAME_CHECKED(UOutlinerSettings, ItemText))
 		{
 			UOutlinerSettings* OutlinerSettings = GetMutableDefault<UOutlinerSettings>();
diff --git a/Engine/Plugins/Experimental/ChaosEditor/Source/FractureEditor/Private/SGeometryCollectionHistogram.cpp b/Engine/Plugins/Experimental/ChaosEditor/Source/FractureEditor/Private/SGeometryCollectionHistogram.cpp
--- a/Engine/Plugins/Experimental/ChaosEditor/Source/FractureEditor/Private/SGeometryCollectionHistogram.cpp
+++ b/Engine/Plugins/Experimental/ChaosEditor/Source/FractureEditor/Private/SGeometryCollectionHistogram.cpp
@@ -118,6 +118,14 @@ FGeometryCollectionHistogramItemPtr FGeometryCollectionHistogramItemComponent::G
 
 FGeometryCollectionHistogramItemList FGeometryCollectionHistogramItemComponent::RegenerateNodes(int32 LevelView)
 {
+	// Filter nodes by simulation type
+	UHistogramSettings* HistogramSettings = GetMutableDefault<UHistogramSettings>();
+	TArray<bool> FilterNodeFlags;
+	FilterNodeFlags.SetNum(FGeometryCollection::ESimulationTypes::FST_Max);
+	FilterNodeFlags[FGeometryCollection::ESimulationTypes::FST_None] = HistogramSettings->bShowEmbedded;
+	FilterNodeFlags[FGeometryCollection::ESimulationTypes::FST_Rigid] = HistogramSettings->bShowRigids;
+	FilterNodeFlags[FGeometryCollection::ESimulationTypes::FST_Clustered] = HistogramSettings->bShowClusters;
+	
 	// Collect the inspected attribute 
 	FGeometryCollectionHistogramItemList NodesList;
 	
@@ -139,7 +147,7 @@ FGeometryCollectionHistogramItemList FGeometryCollectionHistogramItemComponent::
 			// Add a sub item to the histogram for each of the geometry nodes in this GeometryCollection
 			for (int32 Index = 0; Index < NumElements; Index++)
 			{
-				if (Collection->SimulationType[Index] == FGeometryCollection::ESimulationTypes::FST_Rigid)
+				if (FilterNodeFlags[Collection->SimulationType[Index]])
 				{
 					if (LevelView > -1)
 					{
@@ -255,12 +263,16 @@ void SGeometryCollectionHistogram::RefreshView(bool bSorted)
 
 void SGeometryCollectionHistogram::RegenerateNodes(int32 LevelView)
 {
+	LeafNodes.Empty();
 	for (TSharedPtr<FGeometryCollectionHistogramItemComponent> Root : RootNodes)
 	{
-		Root->RegenerateNodes(LevelView);
+		LeafNodes.Append(Root->RegenerateNodes(LevelView));
 	}
+	SetListIndices();
 
 	UHistogramSettings* HistogramSettings = GetMutableDefault<UHistogramSettings>();
+	InspectAttribute(HistogramSettings->InspectedAttribute);
+	NormalizeInspectedValues();
 	RefreshView(HistogramSettings->bSorted);
 }
 
diff --git a/Engine/Plugins/Experimental/ChaosEditor/Source/FractureEditor/Private/SGeometryCollectionHistogram.h b/Engine/Plugins/Experimental/ChaosEditor/Source/FractureEditor/Private/SGeometryCollectionHistogram.h
--- a/Engine/Plugins/Experimental/ChaosEditor/Source/FractureEditor/Private/SGeometryCollectionHistogram.h
+++ b/Engine/Plugins/Experimental/ChaosEditor/Source/FractureEditor/Private/SGeometryCollectionHistogram.h
@@ -43,11 +43,23 @@ public:
 
 	/** What attribute are we inspecting? */
 	UPROPERTY(EditAnywhere, Category = HistogramSettings, meta = (DisplayName = "Inspected Attribute"))
-		EInspectedAttributeEnum InspectedAttribute;
+	EInspectedAttributeEnum InspectedAttribute;
 
 	/** Sort the values? */
 	UPROPERTY(EditAnywhere, Category = HistogramSettings, meta = (DisplayName = "Sort Values"))
-		bool bSorted;
+	bool bSorted;
+
+	/** Show clusters? */
+	UPROPERTY(EditAnywhere, Category = HistogramSettings, meta = (DisplayName = "Show Clusters"))
+	bool bShowClusters=true;
+
+	/** Show rigids? */
+	UPROPERTY(EditAnywhere, Category = HistogramSettings, meta = (DisplayName = "Show Rigids"))
+	bool bShowRigids=true;
+
+	/** Show embedded geometry? */
+	UPROPERTY(EditAnywhere, Category = HistogramSettings, meta = (DisplayName = "Show Embedded Geometry"))
+	bool bShowEmbedded=true;
 
 };
 
diff --git a/Engine/Source/Runtime/Experimental/Chaos/Public/GeometryCollection/GeometryCollection.h b/Engine/Source/Runtime/Experimental/Chaos/Public/GeometryCollection/GeometryCollection.h
--- a/Engine/Source/Runtime/Experimental/Chaos/Public/GeometryCollection/GeometryCollection.h
+++ b/Engine/Source/Runtime/Experimental/Chaos/Public/GeometryCollection/GeometryCollection.h
@@ -96,7 +96,9 @@ public:
 	{
 		FST_None = 0,
 		FST_Rigid = 1,
-		FST_Clustered = 2
+		FST_Clustered = 2,
+
+		FST_Max = 3
 	};
 
 	enum ENodeFlags : uint32
