Fracture Tool -- test for nullptr to avoid edge case crashes

diff --git a/Engine/Plugins/Experimental/ChaosEditor/Source/FractureEditor/Private/FractureToolEmbed.cpp b/Engine/Plugins/Experimental/ChaosEditor/Source/FractureEditor/Private/FractureToolEmbed.cpp
--- a/Engine/Plugins/Experimental/ChaosEditor/Source/FractureEditor/Private/FractureToolEmbed.cpp
+++ b/Engine/Plugins/Experimental/ChaosEditor/Source/FractureEditor/Private/FractureToolEmbed.cpp
@@ -237,7 +237,7 @@ void UFractureToolAutoEmbedGeometry::Execute(TWeakPtr<FFractureEditorModeToolkit
 			}
 		}
 
-		if (ClosestGeometryCollection.IsValid())
+		if (ClosestGeometryCollection.IsValid() && ClosestComponent)
 		{
 			// Which bone points to the closest convex?
 			const TManagedArray<FTransform>& Transform = ClosestGeometryCollection->Transform;
