Fracture makes use of new Non-Exclusive,  Non-Integrated Tool Palettes

diff --git a/Engine/Plugins/Experimental/ChaosEditor/Source/FractureEditor/Public/FractureEditorModeToolkit.h b/Engine/Plugins/Experimental/ChaosEditor/Source/FractureEditor/Public/FractureEditorModeToolkit.h
--- a/Engine/Plugins/Experimental/ChaosEditor/Source/FractureEditor/Public/FractureEditorModeToolkit.h
+++ b/Engine/Plugins/Experimental/ChaosEditor/Source/FractureEditor/Public/FractureEditorModeToolkit.h
@@ -142,6 +142,12 @@ public:
 	virtual void GetToolPaletteNames( TArray<FName>& InPaletteName ) const { InPaletteName = PaletteNames; }
 	virtual FText GetToolPaletteDisplayName(FName PaletteName) const; 
 
+	/* Exclusive Tool Palettes only allow users to use tools from one palette at a time */
+	virtual bool HasExclusiveToolPalettes() const { return false; }
+
+	/* Integrated Tool Palettes show up in the same panel as their details */
+	virtual bool HasIntegratedToolPalettes() const { return false; }
+
 	void SetInitialPalette();
 	virtual void BuildToolPalette(FName PaletteName, class FToolBarBuilder& ToolbarBuilder);
 	virtual void OnToolPaletteChanged(FName PaletteName) override;
