Fracture Mode -- Disable bone colors when exiting Facture Mode or when deselecting Geometry Collection actor.

diff --git a/Engine/Plugins/Experimental/ChaosEditor/Source/FractureEditor/Private/FractureEditorMode.cpp b/Engine/Plugins/Experimental/ChaosEditor/Source/FractureEditor/Private/FractureEditorMode.cpp
--- a/Engine/Plugins/Experimental/ChaosEditor/Source/FractureEditor/Private/FractureEditorMode.cpp
+++ b/Engine/Plugins/Experimental/ChaosEditor/Source/FractureEditor/Private/FractureEditorMode.cpp
@@ -338,6 +338,9 @@ void FFractureEditorMode::OnActorSelectionChanged(const TArray<UObject*>& NewSel
 
 			FScopedColorEdit ShowBoneColorsEdit(ExistingSelection);
 			ShowBoneColorsEdit.SetEnableBoneSelection(false);
+			ShowBoneColorsEdit.SetShowBoneColors(false);
+			
+			ExistingSelection->MarkRenderStateDirty();
 		}
 	}
 
