Fracture Mode -- Shortened selection tool names to prevent cutoff of text.

diff --git a/Engine/Plugins/Experimental/ChaosEditor/Source/FractureEditor/Private/FractureToolSelectors.cpp b/Engine/Plugins/Experimental/ChaosEditor/Source/FractureEditor/Private/FractureToolSelectors.cpp
--- a/Engine/Plugins/Experimental/ChaosEditor/Source/FractureEditor/Private/FractureToolSelectors.cpp
+++ b/Engine/Plugins/Experimental/ChaosEditor/Source/FractureEditor/Private/FractureToolSelectors.cpp
@@ -30,7 +30,7 @@ FSlateIcon UFractureToolSelectAll::GetToolIcon() const
 
 void UFractureToolSelectAll::RegisterUICommand(FFractureEditorCommands* BindingContext)
 {
-	UI_COMMAND_EXT(BindingContext, UICommandInfo, "SelectAll", "Select All", "Selects all Bones in the GeometryCollection.", EUserInterfaceActionType::Button, FInputChord(EModifierKey::Control | EModifierKey::Shift, EKeys::A));
+	UI_COMMAND_EXT(BindingContext, UICommandInfo, "SelectAll", "All", "Selects all Bones in the GeometryCollection.", EUserInterfaceActionType::Button, FInputChord(EModifierKey::Control | EModifierKey::Shift, EKeys::A));
 	BindingContext->SelectAll = UICommandInfo;
 }
 
@@ -83,7 +83,7 @@ FSlateIcon UFractureToolSelectNone::GetToolIcon() const
 
 void UFractureToolSelectNone::RegisterUICommand(FFractureEditorCommands* BindingContext)
 {
-	UI_COMMAND_EXT(BindingContext, UICommandInfo, "SelectNone", "Select None", "Deselects all Bones in the GeometryCollection.", EUserInterfaceActionType::Button, FInputChord(EModifierKey::Control | EModifierKey::Shift, EKeys::D));
+	UI_COMMAND_EXT(BindingContext, UICommandInfo, "SelectNone", "None", "Deselects all Bones in the GeometryCollection.", EUserInterfaceActionType::Button, FInputChord(EModifierKey::Control | EModifierKey::Shift, EKeys::D));
 	BindingContext->SelectNone = UICommandInfo;
 }
 
@@ -113,7 +113,7 @@ FSlateIcon UFractureToolSelectNeighbors::GetToolIcon() const
 
 void UFractureToolSelectNeighbors::RegisterUICommand(FFractureEditorCommands* BindingContext)
 {
-	UI_COMMAND_EXT(BindingContext, UICommandInfo, "SelectNeighbors", "Select Neighbors", "Select all bones adjacent to the currently selected bones.", EUserInterfaceActionType::Button, FInputChord());
+	UI_COMMAND_EXT(BindingContext, UICommandInfo, "SelectNeighbors", "Contact", "Select all bones adjacent to the currently selected bones.", EUserInterfaceActionType::Button, FInputChord());
 	BindingContext->SelectNeighbors = UICommandInfo;
 }
 
@@ -143,7 +143,7 @@ FSlateIcon UFractureToolSelectSiblings::GetToolIcon() const
 
 void UFractureToolSelectSiblings::RegisterUICommand(FFractureEditorCommands* BindingContext)
 {
-	UI_COMMAND_EXT(BindingContext, UICommandInfo, "SelectSiblings", "Select Siblings", "Select all bones at the same levels as the currently selected bones.", EUserInterfaceActionType::Button, FInputChord());
+	UI_COMMAND_EXT(BindingContext, UICommandInfo, "SelectSiblings", "Siblings", "Select all bones at the same levels as the currently selected bones.", EUserInterfaceActionType::Button, FInputChord());
 	BindingContext->SelectSiblings = UICommandInfo;
 }
 
@@ -173,7 +173,7 @@ FSlateIcon UFractureToolSelectAllInCluster::GetToolIcon() const
 
 void UFractureToolSelectAllInCluster::RegisterUICommand(FFractureEditorCommands* BindingContext)
 {
-	UI_COMMAND_EXT(BindingContext, UICommandInfo, "SelectAllInCluster", "Select All In Cluster", "Select all bones with the same parent as selected bones.", EUserInterfaceActionType::Button, FInputChord());
+	UI_COMMAND_EXT(BindingContext, UICommandInfo, "SelectAllInCluster", "Cluster", "Select all bones with the same parent as selected bones.", EUserInterfaceActionType::Button, FInputChord());
 	BindingContext->SelectAllInCluster = UICommandInfo;
 }
 
@@ -203,7 +203,7 @@ FSlateIcon UFractureToolSelectInvert::GetToolIcon() const
 
 void UFractureToolSelectInvert::RegisterUICommand(FFractureEditorCommands* BindingContext)
 {
-	UI_COMMAND_EXT(BindingContext, UICommandInfo, "SelectInvert", "Invert Selection", "Invert current selection of bones.", EUserInterfaceActionType::Button, FInputChord());
+	UI_COMMAND_EXT(BindingContext, UICommandInfo, "SelectInvert", "Invert", "Invert current selection of bones.", EUserInterfaceActionType::Button, FInputChord());
 	BindingContext->SelectInvert = UICommandInfo;
 }
 
