GeometryCollection -- Added EmbeddedGeometry: the ability for somulated bones to carry instanced geometry that doesn't contribute to simulation but tracks with evolving rigid body transformation.

diff --git a/Engine/Plugins/Experimental/ChaosEditor/Source/FractureEditor/Private/FractureToolEmbed.cpp b/Engine/Plugins/Experimental/ChaosEditor/Source/FractureEditor/Private/FractureToolEmbed.cpp
--- a/Engine/Plugins/Experimental/ChaosEditor/Source/FractureEditor/Private/FractureToolEmbed.cpp
+++ b/Engine/Plugins/Experimental/ChaosEditor/Source/FractureEditor/Private/FractureToolEmbed.cpp
@@ -49,82 +49,60 @@ bool UFractureToolAddEmbeddedGeometry::CanExecute() const
 
 void UFractureToolAddEmbeddedGeometry::Execute(TWeakPtr<FFractureEditorModeToolkit> InToolkit)
 {
-	if (!InToolkit.IsValid())
+	if (InToolkit.IsValid())
 	{
-		return;
-	}
+		FFractureEditorModeToolkit* Toolkit = InToolkit.Pin().Get();
 
-	// Get selected static meshes
-	TArray<UStaticMeshComponent*> SelectedStaticMeshComponents = GetSelectedStaticMeshComponents();
+		// Get selected static meshes
+		TArray<UStaticMeshComponent*> SelectedStaticMeshComponents = GetSelectedStaticMeshComponents();
 
-	TArray<FFractureToolContext> Contexts = GetFractureToolContexts();
+		TArray<FFractureToolContext> Contexts = GetFractureToolContexts();
 
-	for (FFractureToolContext& Context : Contexts)
-	{
-		FGeometryCollectionEdit GeometryCollectionEdit = Context.GetGeometryCollectionComponent()->EditRestCollection(GeometryCollection::EEditUpdate::RestPhysicsDynamic);
-		UGeometryCollection* FracturedGeometryCollection = GeometryCollectionEdit.GetRestCollection();
-
-		FFractureToolContext::FGeometryCollectionPtr GeometryCollection = Context.GetGeometryCollection();
-		const TManagedArray<FTransform>& Transform = GeometryCollection->Transform;
-		TManagedArray<int32>& Parent = GeometryCollection->Parent;
-		TManagedArray<int32>& SimType = GeometryCollection->SimulationType;
-		TManagedArray<TSet<int32>>& Children = GeometryCollection->Children;
-		TManagedArray<bool>& CollectionSimulatableParticles =
-			GeometryCollection->GetAttribute<bool>(
-				FGeometryCollection::SimulatableParticlesAttribute, FTransformCollection::TransformGroup);
-
-		const FTransform TargetActorTransform(Context.GetGeometryCollectionComponent()->GetOwner()->GetTransform());
-
-		Context.ConvertSelectionToRigidNodes();
-		const TArray<int32>& SelectedBones = Context.GetSelection();
-		for (const int32 SelectedBone : SelectedBones)
+		for (FFractureToolContext& Context : Contexts)
 		{
-			FTransform BoneGlobalTransform = GeometryCollectionAlgo::GlobalMatrix(Transform, Parent, SelectedBone);
+			FGeometryCollectionEdit GeometryCollectionEdit = Context.GetGeometryCollectionComponent()->EditRestCollection(GeometryCollection::EEditUpdate::RestPhysicsDynamic);
+			UGeometryCollection* FracturedGeometryCollection = GeometryCollectionEdit.GetRestCollection();
+
+			FFractureToolContext::FGeometryCollectionPtr GeometryCollection = Context.GetGeometryCollection();
+			const TManagedArray<FTransform>& Transform = GeometryCollection->Transform;
+			const TManagedArray<int32>& Parent = GeometryCollection->Parent;
 			
-			for (UStaticMeshComponent* SelectedStaticMeshComponent : SelectedStaticMeshComponents)
-			{
-				const AActor* Actor = SelectedStaticMeshComponent->GetOwner();
-				check(Actor)
+			const FTransform TargetActorTransform(Context.GetGeometryCollectionComponent()->GetOwner()->GetTransform());
 
-				const FTransform SMActorTransform(Actor->GetTransform());
+			Context.ConvertSelectionToRigidNodes();
+			const TArray<int32>& SelectedBones = Context.GetSelection();
+			for (const int32 SelectedBone : SelectedBones)
+			{
+				FTransform BoneGlobalTransform = GeometryCollectionAlgo::GlobalMatrix(Transform, Parent, SelectedBone);
 
-				UStaticMesh* ComponentStaticMesh = SelectedStaticMeshComponent->GetStaticMesh();
-				if (ComponentStaticMesh)
+				for (UStaticMeshComponent* SelectedStaticMeshComponent : SelectedStaticMeshComponents)
 				{
-					// If any of the static meshes have Nanite enabled, also enable on the new geometry collection asset for convenience.
-					FracturedGeometryCollection->EnableNanite |= ComponentStaticMesh->NaniteSettings.bEnabled;
-				}
-
-				FTransform ComponentTransform = SMActorTransform.GetRelativeTransform(TargetActorTransform);
-				FTransform BoneTransform = ComponentTransform.GetRelativeTransform(BoneGlobalTransform);
-
-				decltype(FGeometryCollectionSource::SourceMaterial) SourceMaterials(SelectedStaticMeshComponent->GetMaterials());
+					const AActor* Actor = SelectedStaticMeshComponent->GetOwner();
+					check(Actor)
 
-				// Add static mesh geometry to context's geometry collection.
-				FGeometryCollectionConversion::AppendStaticMesh(ComponentStaticMesh, SelectedStaticMeshComponent, BoneTransform, FracturedGeometryCollection, true);
+					const FTransform SMActorTransform(Actor->GetTransform());
 
-				int32 NewTransformIndex = Parent.Num() - 1;
-				check(NewTransformIndex > -1)
+					UStaticMesh* ComponentStaticMesh = SelectedStaticMeshComponent->GetStaticMesh();
+					
+					FTransform ComponentTransform = SMActorTransform.GetRelativeTransform(TargetActorTransform);
+					FTransform BoneTransform = ComponentTransform.GetRelativeTransform(BoneGlobalTransform);
 
-				// Reparent the new transform to the selected bone
-				Parent[NewTransformIndex] = SelectedBone;
-				Children[SelectedBone].Add(NewTransformIndex);
-
-				// Set properties appropriate to embedded geometry
-				SimType[NewTransformIndex] = FGeometryCollection::ESimulationTypes::FST_None;
-				CollectionSimulatableParticles[NewTransformIndex] = false;
+					int32 ExemplarIndex = FracturedGeometryCollection->AttachEmbeddedGeometryExemplar(ComponentStaticMesh);
+					if (GeometryCollection->AppendEmbeddedInstance(ExemplarIndex, SelectedBone, BoneTransform))
+					{
+						FracturedGeometryCollection->EmbeddedGeometryExemplar[ExemplarIndex].InstanceCount++;
+					}
+				}
 			}
-		}
 
-		FracturedGeometryCollection->InitializeMaterials();
+			Context.GetGeometryCollectionComponent()->InitializeEmbeddedGeometry();
+			Refresh(Context, Toolkit);
 
-		FGeometryCollectionClusteringUtility::UpdateHierarchyLevelOfChildren(GeometryCollection.Get(), -1);
-		Context.GetGeometryCollectionComponent()->MarkRenderStateDirty();
+			FracturedGeometryCollection->MarkPackageDirty();
+		}
 
-		FracturedGeometryCollection->MarkPackageDirty();
+		SetOutlinerComponents(Contexts, Toolkit);
 	}
-
-	SetOutlinerComponents(Contexts, InToolkit.Pin().Get());
 }
 
 TArray<UStaticMeshComponent*> UFractureToolAddEmbeddedGeometry::GetSelectedStaticMeshComponents()
@@ -183,124 +161,104 @@ bool UFractureToolAutoEmbedGeometry::CanExecute() const
 
 void UFractureToolAutoEmbedGeometry::Execute(TWeakPtr<FFractureEditorModeToolkit> InToolkit)
 {
-	if (!InToolkit.IsValid())
+	if (InToolkit.IsValid())
 	{
-		return;
-	}
+		FFractureEditorModeToolkit* Toolkit = InToolkit.Pin().Get();
 
-	// Get selected static meshes
-	TArray<UStaticMeshComponent*> SelectedStaticMeshComponents = GetSelectedStaticMeshComponents();
-	TArray<FFractureToolContext> Contexts = GetFractureToolContexts();
+		// Get selected static meshes
+		TArray<UStaticMeshComponent*> SelectedStaticMeshComponents = GetSelectedStaticMeshComponents();
+		TArray<FFractureToolContext> Contexts = GetFractureToolContexts();
 
-	// For each static mesh component, we iterate all the convex hulls and determine which convex hull best contains the worldspace pivot of the static mesh.
-	for (UStaticMeshComponent* SelectedStaticMeshComponent : SelectedStaticMeshComponents)
-	{
-		// Static Mesh world space location
-		FVector SMLocation = SelectedStaticMeshComponent->GetComponentLocation();
-
-		// Determine the "closest" bone
-		FFractureToolContext::FGeometryCollectionPtr ClosestGeometryCollection;
-		int32 ClosestConvex = INDEX_NONE;
-		Chaos::FReal ClosestPhi = TNumericLimits<float>::Max();
-		UGeometryCollectionComponent* ClosestComponent = nullptr;
-		
-		for (FFractureToolContext& Context : Contexts)
+		// For each static mesh component, we iterate all the convex hulls and determine which convex hull best contains the worldspace pivot of the static mesh.
+		for (UStaticMeshComponent* SelectedStaticMeshComponent : SelectedStaticMeshComponents)
 		{
-			FGeometryCollection* GeometryCollection = Context.GetGeometryCollection().Get();
-			FGeometryCollectionConvexUtility::FGeometryCollectionConvexData ConvexData = FGeometryCollectionConvexUtility::GetValidConvexHullData(GeometryCollection);
-			const TManagedArray<FTransform>& Transform = GeometryCollection->Transform;
-			const TManagedArray<int32>& Parent = GeometryCollection->Parent;
-			TArray<FTransform> BoneGlobalTransforms;
-			GeometryCollectionAlgo::GlobalMatrices(Transform, Parent, BoneGlobalTransforms);
-			
-			FTransform WorldToComponent = Context.GetGeometryCollectionComponent()->GetComponentToWorld().Inverse();
-			FVector ComponentSpaceLocation = WorldToComponent.TransformPosition(SMLocation);
+			// Static Mesh world space location
+			FVector SMLocation = SelectedStaticMeshComponent->GetComponentLocation();
+
+			// Determine the "closest" bone
+			FFractureToolContext::FGeometryCollectionPtr ClosestGeometryCollection;
+			int32 ClosestConvex = INDEX_NONE;
+			Chaos::FReal ClosestPhi = TNumericLimits<float>::Max();
+			UGeometryCollectionComponent* ClosestComponent = nullptr;
+			FFractureToolContext* ClosestContext = nullptr;
 
-			int32 NumConvex = ConvexData.ConvexHull.Num();
-			for (int32 ConvexIndex = 0; ConvexIndex < NumConvex; ++ConvexIndex)
+			for (FFractureToolContext& Context : Contexts)
 			{
-				// transform into bone space where the convex hull is described
-				int32 TransformIndex = ConvexData.TransformToConvexIndex.Find(ConvexIndex);
-				FTransform ComponentToBone = BoneGlobalTransforms[TransformIndex].Inverse();
-				FVector BoneSpaceLocation = ComponentToBone.TransformPosition(ComponentSpaceLocation);
-				
-				Chaos::FVec3 Normal;
-				Chaos::FReal Phi = ConvexData.ConvexHull[ConvexIndex]->PhiWithNormal(BoneSpaceLocation, Normal);
-				UE_LOG(LogTemp, Warning, TEXT(" Cone %d Phi %f"), ConvexData.TransformToConvexIndex.Find(ConvexIndex), Phi);
-				if (Phi < ClosestPhi)
+				FGeometryCollection* GeometryCollection = Context.GetGeometryCollection().Get();
+				FGeometryCollectionConvexUtility::FGeometryCollectionConvexData ConvexData = FGeometryCollectionConvexUtility::GetValidConvexHullData(GeometryCollection);
+				const TManagedArray<FTransform>& Transform = GeometryCollection->Transform;
+				const TManagedArray<int32>& Parent = GeometryCollection->Parent;
+				TArray<FTransform> BoneGlobalTransforms;
+				GeometryCollectionAlgo::GlobalMatrices(Transform, Parent, BoneGlobalTransforms);
+
+				FTransform WorldToComponent = Context.GetGeometryCollectionComponent()->GetComponentToWorld().Inverse();
+				FVector ComponentSpaceLocation = WorldToComponent.TransformPosition(SMLocation);
+
+				int32 NumConvex = ConvexData.ConvexHull.Num();
+				for (int32 ConvexIndex = 0; ConvexIndex < NumConvex; ++ConvexIndex)
 				{
-					ClosestPhi = Phi;
-					ClosestGeometryCollection = Context.GetGeometryCollection();
-					ClosestConvex = ConvexIndex;
-					ClosestComponent = Context.GetGeometryCollectionComponent();
+					// transform into bone space where the convex hull is described
+					int32 TransformIndex = ConvexData.TransformToConvexIndex.Find(ConvexIndex);
+					FTransform ComponentToBone = BoneGlobalTransforms[TransformIndex].Inverse();
+					FVector BoneSpaceLocation = ComponentToBone.TransformPosition(ComponentSpaceLocation);
+
+					Chaos::FVec3 Normal;
+					Chaos::FReal Phi = ConvexData.ConvexHull[ConvexIndex]->PhiWithNormal(BoneSpaceLocation, Normal);
+					UE_LOG(LogTemp, Warning, TEXT(" Cone %d Phi %f"), ConvexData.TransformToConvexIndex.Find(ConvexIndex), Phi);
+					if (Phi < ClosestPhi)
+					{
+						ClosestPhi = Phi;
+						ClosestGeometryCollection = Context.GetGeometryCollection();
+						ClosestConvex = ConvexIndex;
+						ClosestComponent = Context.GetGeometryCollectionComponent();
+						ClosestContext = &Context;
+					}
 				}
 			}
-		}
 
-		if (ClosestGeometryCollection.IsValid() && ClosestComponent)
-		{
-			// Which bone points to the closest convex?
-			const TManagedArray<FTransform>& Transform = ClosestGeometryCollection->Transform;
-			TManagedArray<int32>& Parent = ClosestGeometryCollection->Parent;
-			TManagedArray<int32>& SimType = ClosestGeometryCollection->SimulationType;
-			TManagedArray<TSet<int32>>& Children = ClosestGeometryCollection->Children;
-			TManagedArray<bool>& CollectionSimulatableParticles =
-				ClosestGeometryCollection->GetAttribute<bool>(
-					FGeometryCollection::SimulatableParticlesAttribute, FTransformCollection::TransformGroup);
-			TManagedArray<int32>& TransformToConvexIndex =
-				ClosestGeometryCollection->GetAttribute<int32>("TransformToConvexIndex", FTransformCollection::TransformGroup);
-
-			const int32 BoneIndex = TransformToConvexIndex.Find(ClosestConvex);
-
-			if (BoneIndex > INDEX_NONE)
+			if (ClosestGeometryCollection.IsValid() && ClosestComponent && ClosestContext)
 			{
-				// We found the closest bone, now we embed the geometry.
-				FGeometryCollectionEdit GeometryCollectionEdit = ClosestComponent->EditRestCollection(GeometryCollection::EEditUpdate::RestPhysicsDynamic);
-				UGeometryCollection* FracturedGeometryCollection = GeometryCollectionEdit.GetRestCollection();
-
-				FTransform BoneGlobalTransform = GeometryCollectionAlgo::GlobalMatrix(Transform, Parent, BoneIndex);
+				// Which bone points to the closest convex?
+				const TManagedArray<FTransform>& Transform = ClosestGeometryCollection->Transform;
+				const TManagedArray<int32>& Parent = ClosestGeometryCollection->Parent;
+				TManagedArray<int32>& TransformToConvexIndex =
+					ClosestGeometryCollection->GetAttribute<int32>("TransformToConvexIndex", FTransformCollection::TransformGroup);
 
-				const FTransform TargetActorTransform(ClosestComponent->GetOwner()->GetTransform());
-				const AActor* Actor = SelectedStaticMeshComponent->GetOwner();
-				check(Actor)
+				const int32 BoneIndex = TransformToConvexIndex.Find(ClosestConvex);
 
-				const FTransform SMActorTransform(Actor->GetTransform());
-				FTransform ComponentTransform = SMActorTransform.GetRelativeTransform(TargetActorTransform);
-				FTransform BoneTransform = ComponentTransform.GetRelativeTransform(BoneGlobalTransform);
-
-				UStaticMesh* ComponentStaticMesh = SelectedStaticMeshComponent->GetStaticMesh();
-				if (ComponentStaticMesh)
+				if (BoneIndex > INDEX_NONE)
 				{
-					// If any of the static meshes have Nanite enabled, also enable on the new geometry collection asset for convenience.
-					FracturedGeometryCollection->EnableNanite |= ComponentStaticMesh->NaniteSettings.bEnabled;
-				}
+					// We found the closest bone, now we embed the geometry.
+					FGeometryCollectionEdit GeometryCollectionEdit = ClosestComponent->EditRestCollection(GeometryCollection::EEditUpdate::RestPhysicsDynamic);
+					UGeometryCollection* FracturedGeometryCollection = GeometryCollectionEdit.GetRestCollection();
 
-				// Add static mesh geometry to context's geometry collection.
-				FGeometryCollectionConversion::AppendStaticMesh(ComponentStaticMesh, SelectedStaticMeshComponent, BoneTransform, FracturedGeometryCollection, true);
+					FTransform BoneGlobalTransform = GeometryCollectionAlgo::GlobalMatrix(Transform, Parent, BoneIndex);
 
-				// A new transform has been appended to the TransformGroup
-				int32 NewTransformIndex = Parent.Num() - 1;
-				check(NewTransformIndex > -1)
+					const FTransform TargetActorTransform(ClosestComponent->GetOwner()->GetTransform());
+					const AActor* Actor = SelectedStaticMeshComponent->GetOwner();
+					check(Actor)
 
-				// Reparent the new transform to the selected bone
-				Parent[NewTransformIndex] = BoneIndex;
-				Children[BoneIndex].Add(NewTransformIndex);
+					const FTransform SMActorTransform(Actor->GetTransform());
+					FTransform ComponentTransform = SMActorTransform.GetRelativeTransform(TargetActorTransform);
+					FTransform BoneTransform = ComponentTransform.GetRelativeTransform(BoneGlobalTransform);
 
-				// Set properties appropriate to embedded geometry
-				SimType[NewTransformIndex] = FGeometryCollection::ESimulationTypes::FST_None;
-				CollectionSimulatableParticles[NewTransformIndex] = false;
-
-				FracturedGeometryCollection->InitializeMaterials();
-
-				FGeometryCollectionClusteringUtility::UpdateHierarchyLevelOfChildren(ClosestGeometryCollection.Get(), -1);
-				ClosestComponent->MarkRenderStateDirty();
+					UStaticMesh* ComponentStaticMesh = SelectedStaticMeshComponent->GetStaticMesh();
+					int32 ExemplarIndex = FracturedGeometryCollection->AttachEmbeddedGeometryExemplar(ComponentStaticMesh);
+					if (ClosestGeometryCollection->AppendEmbeddedInstance(ExemplarIndex, BoneIndex, BoneTransform))
+					{
+						FracturedGeometryCollection->EmbeddedGeometryExemplar[ExemplarIndex].InstanceCount++;
+					}
 
-				FracturedGeometryCollection->MarkPackageDirty();
+					// #todo there might be a lot of these -- collect and put outside the loop.
+					ClosestComponent->InitializeEmbeddedGeometry();
+					Refresh(*ClosestContext, Toolkit);
+					FracturedGeometryCollection->MarkPackageDirty();
+				}
 			}
 		}
-	}
 
-	SetOutlinerComponents(Contexts, InToolkit.Pin().Get());
+		SetOutlinerComponents(Contexts, Toolkit);
+	}
 }
 
 TArray<UStaticMeshComponent*> UFractureToolAutoEmbedGeometry::GetSelectedStaticMeshComponents()
@@ -360,56 +318,72 @@ bool UFractureToolDeleteEmbeddedGeometry::CanExecute() const
 
 void UFractureToolDeleteEmbeddedGeometry::Execute(TWeakPtr<FFractureEditorModeToolkit> InToolkit)
 {
-	if (!InToolkit.IsValid())
+	if (InToolkit.IsValid())
 	{
-		return;
-	}
+		FFractureEditorModeToolkit* Toolkit = InToolkit.Pin().Get();
 
-	// Get selected static meshes
-	TArray<FFractureToolContext> Contexts = GetFractureToolContexts();
+		// Get selected static meshes
+		TArray<FFractureToolContext> Contexts = GetFractureToolContexts();
 
-	for (FFractureToolContext& Context : Contexts)
-	{
-		Context.Sanitize();
-		
-		FGeometryCollectionEdit GeometryCollectionEdit = Context.GetGeometryCollectionComponent()->EditRestCollection(GeometryCollection::EEditUpdate::RestPhysicsDynamic);
-		UGeometryCollection* FracturedGeometryCollection = GeometryCollectionEdit.GetRestCollection();
+		for (FFractureToolContext& Context : Contexts)
+		{
+			Context.Sanitize();
 
-		FFractureToolContext::FGeometryCollectionPtr GeometryCollection = Context.GetGeometryCollection();
-		const TManagedArray<int32>& SimType = GeometryCollection->SimulationType;
+			FGeometryCollectionEdit GeometryCollectionEdit = Context.GetGeometryCollectionComponent()->EditRestCollection(GeometryCollection::EEditUpdate::RestPhysicsDynamic);
+			UGeometryCollection* FracturedGeometryCollection = GeometryCollectionEdit.GetRestCollection();
 
-		TArray<int32> EmbeddedGeometryToBeRemoved;
-		const TArray<int32>& SelectedBones = Context.GetSelection();
-		for (const int32 SelectedBone : SelectedBones)
-		{
-			if (SimType[SelectedBone] != FGeometryCollection::ESimulationTypes::FST_None)
+			FFractureToolContext::FGeometryCollectionPtr GeometryCollection = Context.GetGeometryCollection();
+			const TManagedArray<int32>& ExemplarIndex = GeometryCollection->ExemplarIndex;
+
+			TArray<int32> EmbeddedGeometryToBeRemoved;
+			const TArray<int32>& SelectedBones = Context.GetSelection();
+			for (const int32 SelectedBone : SelectedBones)
 			{
-				// Select all Embedded Geometry found in the selected branch.
-				TArray<int32> LeafBones;
-				FGeometryCollectionClusteringUtility::GetLeafBones(GeometryCollection.Get(), SelectedBone, false, LeafBones);
-				for (int32 LeafBone : LeafBones)
+				if (ExemplarIndex[SelectedBone] == INDEX_NONE)
 				{
-					if (SimType[LeafBone] == FGeometryCollection::ESimulationTypes::FST_None)
+					// Select all Embedded Geometry found in the selected branch.
+					TArray<int32> LeafBones;
+					FGeometryCollectionClusteringUtility::GetLeafBones(GeometryCollection.Get(), SelectedBone, false, LeafBones);
+					for (int32 LeafBone : LeafBones)
 					{
-						EmbeddedGeometryToBeRemoved.Add(LeafBone);
+						if (ExemplarIndex[LeafBone] > INDEX_NONE)
+						{
+							EmbeddedGeometryToBeRemoved.Add(LeafBone);
+						}
 					}
 				}
+				else
+				{
+					// Selected bone is embedded geometry. Only delete this.
+					EmbeddedGeometryToBeRemoved.Add(SelectedBone);
+				}
 			}
-			else
+
+			TArray<int32> UninstancedExemplars;
+			UninstancedExemplars.Reserve(EmbeddedGeometryToBeRemoved.Num());
+			for (int32 EmbeddedIndex : EmbeddedGeometryToBeRemoved)
 			{
-				// Selected bone is embedded geometry. Only delete this.
-				EmbeddedGeometryToBeRemoved.Add(SelectedBone);
+				if ((--FracturedGeometryCollection->EmbeddedGeometryExemplar[ExemplarIndex[EmbeddedIndex]].InstanceCount) < 1)
+				{
+					UE_LOG(LogFractureTool, Warning, TEXT("Exemplar Index %d is empty. Removing Exemplar from Geometry Collection."), ExemplarIndex[EmbeddedIndex]);
+					UninstancedExemplars.Add(ExemplarIndex[EmbeddedIndex]);
+				}
 			}
-		}
 
-		EmbeddedGeometryToBeRemoved.Sort();
-		GeometryCollection->RemoveElements(FGeometryCollection::TransformGroup, EmbeddedGeometryToBeRemoved);
+			EmbeddedGeometryToBeRemoved.Sort();
+			GeometryCollection->RemoveElements(FGeometryCollection::TransformGroup, EmbeddedGeometryToBeRemoved);
 
-		Context.GetGeometryCollectionComponent()->MarkRenderStateDirty();
-		FracturedGeometryCollection->MarkPackageDirty();
-	}
+			UninstancedExemplars.Sort();
+			FracturedGeometryCollection->RemoveExemplars(UninstancedExemplars);
+			GeometryCollection->ReindexExemplarIndices(UninstancedExemplars);
 
-	SetOutlinerComponents(Contexts, InToolkit.Pin().Get());
+			Context.GetGeometryCollectionComponent()->InitializeEmbeddedGeometry();
+			Refresh(Context, Toolkit, true);
+			FracturedGeometryCollection->MarkPackageDirty();
+		}
+
+		SetOutlinerComponents(Contexts, Toolkit);
+	}
 }
 
 
diff --git a/Engine/Plugins/Experimental/ChaosEditor/Source/FractureEditor/Public/FractureTool.cpp b/Engine/Plugins/Experimental/ChaosEditor/Source/FractureEditor/Public/FractureTool.cpp
--- a/Engine/Plugins/Experimental/ChaosEditor/Source/FractureEditor/Public/FractureTool.cpp
+++ b/Engine/Plugins/Experimental/ChaosEditor/Source/FractureEditor/Public/FractureTool.cpp
@@ -57,7 +57,10 @@ bool UFractureActionTool::IsGeometryCollectionSelected()
 		AActor* Actor = Cast<AActor>(*Iter);
 		if (Actor)
 		{
-			if (Actor->FindComponentByClass<UGeometryCollectionComponent>())
+			TInlineComponentArray<UGeometryCollectionComponent*> GeometryCollectionComponents;
+			Actor->GetComponents<UGeometryCollectionComponent>(GeometryCollectionComponents, true);
+
+			if (GeometryCollectionComponents.Num() > 0)
 			{
 				return true;
 			}
@@ -127,7 +130,7 @@ void UFractureActionTool::GetSelectedGeometryCollectionComponents(TSet<UGeometry
 	}
 }
 
-void UFractureActionTool::Refresh(FFractureToolContext& Context, FFractureEditorModeToolkit* Toolkit)
+void UFractureActionTool::Refresh(FFractureToolContext& Context, FFractureEditorModeToolkit* Toolkit, bool bClearSelection)
 {
 	UGeometryCollectionComponent* GeometryCollectionComponent = Context.GetGeometryCollectionComponent();
 	TSharedPtr<FGeometryCollection, ESPMode::ThreadSafe> GeometryCollectionPtr = Context.GetGeometryCollection();
@@ -138,10 +141,18 @@ void UFractureActionTool::Refresh(FFractureToolContext& Context, FFractureEditor
 	Toolkit->RegenerateHistogram();
 
 	FScopedColorEdit EditBoneColor(GeometryCollectionComponent, true);
-	EditBoneColor.SetSelectedBones(Context.GetSelection());
-
-	FFractureSelectionTools::ToggleSelectedBones(GeometryCollectionComponent, Context.GetSelection(), true);
 
+	if (bClearSelection)
+	{
+		EditBoneColor.ResetBoneSelection();
+		FFractureSelectionTools::ClearSelectedBones(GeometryCollectionComponent);
+	}
+	else
+	{
+		EditBoneColor.SetSelectedBones(Context.GetSelection());
+		FFractureSelectionTools::ToggleSelectedBones(GeometryCollectionComponent, Context.GetSelection(), true);
+	}
+	
 	Toolkit->UpdateExplodedVectors(GeometryCollectionComponent);
 
 	GeometryCollectionComponent->MarkRenderDynamicDataDirty();
@@ -238,12 +249,12 @@ bool UFractureModalTool::CanExecute() const
 	{
 		return false;
 	}
-
+	/*
 	if (IsStaticMeshSelected())
 	{
 		return false;
 	}
-
+	*/
 	return true;
 }
 
diff --git a/Engine/Plugins/Experimental/ChaosEditor/Source/FractureEditor/Public/FractureTool.h b/Engine/Plugins/Experimental/ChaosEditor/Source/FractureEditor/Public/FractureTool.h
--- a/Engine/Plugins/Experimental/ChaosEditor/Source/FractureEditor/Public/FractureTool.h
+++ b/Engine/Plugins/Experimental/ChaosEditor/Source/FractureEditor/Public/FractureTool.h
@@ -68,7 +68,7 @@ protected:
 	static void AddSingleRootNodeIfRequired(UGeometryCollection* GeometryCollectionObject);
 	static void AddAdditionalAttributesIfRequired(UGeometryCollection* GeometryCollectionObject);
 	static void GetSelectedGeometryCollectionComponents(TSet<UGeometryCollectionComponent*>& GeomCompSelection);
-	static void Refresh(FFractureToolContext& Context, FFractureEditorModeToolkit* Toolkit);
+	static void Refresh(FFractureToolContext& Context, FFractureEditorModeToolkit* Toolkit, bool bClearSelection=false);
 	static void SetOutlinerComponents(TArray<FFractureToolContext>& InContexts, FFractureEditorModeToolkit* Toolkit);
 
 
diff --git a/Engine/Source/Runtime/Experimental/Chaos/Private/GeometryCollection/GeometryCollection.cpp b/Engine/Source/Runtime/Experimental/Chaos/Private/GeometryCollection/GeometryCollection.cpp
--- a/Engine/Source/Runtime/Experimental/Chaos/Private/GeometryCollection/GeometryCollection.cpp
+++ b/Engine/Source/Runtime/Experimental/Chaos/Private/GeometryCollection/GeometryCollection.cpp
@@ -49,6 +49,7 @@ void FGeometryCollection::Construct()
 	AddExternalAttribute<int32>("SimulationType", FTransformCollection::TransformGroup, SimulationType);
 	AddExternalAttribute<int32>("StatusFlags", FTransformCollection::TransformGroup, StatusFlags);
 	AddExternalAttribute<int32>("InitialDynamicState", FTransformCollection::TransformGroup, InitialDynamicState);
+	AddExternalAttribute<int32>("ExemplarIndex", FTransformCollection::TransformGroup, ExemplarIndex);
 
 	// Vertices Group
 	AddExternalAttribute<FVector>("Vertex", FGeometryCollection::VerticesGroup, Vertex);
@@ -92,6 +93,7 @@ void FGeometryCollection::SetDefaults(FName Group, uint32 StartSize, uint32 NumE
 			SimulationType[Idx] = FGeometryCollection::ESimulationTypes::FST_None;
 			StatusFlags[Idx] = 0;
 			InitialDynamicState[Idx] = static_cast<int32>(Chaos::EObjectStateType::Uninitialized);
+			ExemplarIndex[Idx] = INDEX_NONE;
 		}
 
 		FGeometryCollectionConvexUtility::SetDefaults(this, Group, StartSize, NumElements);
@@ -143,6 +145,7 @@ int32 FGeometryCollection::AppendGeometry(const FGeometryCollection & Element, i
 	const TManagedArray<int32>& ElementSimulationType = Element.SimulationType;
 	const TManagedArray<int32>& ElementStatusFlags = Element.StatusFlags;
 	const TManagedArray<int32>& ElementInitialDynamicState = Element.InitialDynamicState;
+	const TManagedArray<int32>& ElementExemplarIndex = Element.ExemplarIndex;
 
 	// --- TRANSFORM ---
 	for (int TransformIdx = 0; TransformIdx < NumNewTransforms; TransformIdx++)
@@ -150,6 +153,7 @@ int32 FGeometryCollection::AppendGeometry(const FGeometryCollection & Element, i
 		SimulationType[TransformIdx + StartTransformIndex] = ElementSimulationType[TransformIdx];
 		StatusFlags[TransformIdx + StartTransformIndex] = ElementStatusFlags[TransformIdx];
 		InitialDynamicState[TransformIdx + StartTransformIndex] = ElementInitialDynamicState[TransformIdx];
+		ExemplarIndex[TransformIdx + StartTransformIndex] = ElementExemplarIndex[TransformIdx];
 	}
 
 	// --- VERTICES GROUP ---
@@ -299,6 +303,50 @@ int32 FGeometryCollection::AppendGeometry(const FGeometryCollection & Element, i
 	return StartTransformIndex;
 }
 
+bool FGeometryCollection::AppendEmbeddedInstance(int32 InExemplarIndex, int32 InParentIndex, const FTransform& InTransform)
+{
+	if (InParentIndex == INDEX_NONE || InParentIndex >= NumElements(FGeometryCollection::TransformGroup))
+	{
+		return false;
+	}
+
+	// add a new embedded instance
+	int32 Element = AddElements(1, FGeometryCollection::TransformGroup);
+	Transform[Element] = InTransform;
+	Parent[Element] = InParentIndex;
+	Children[InParentIndex].Add(Element);
+	SimulationType[Element] = FST_None;
+	ExemplarIndex[Element] = InExemplarIndex;
+	TransformToGeometryIndex[Element] = INDEX_NONE;
+
+	return true;
+}
+
+
+void FGeometryCollection::ReindexExemplarIndices(TArray<int32>& SortedRemovedIndices)
+{
+	for (int32 Index = 0; Index < NumElements(TransformGroup); ++Index)
+	{
+		if (ExemplarIndex[Index] > INDEX_NONE)
+		{
+			for (int32 RemovalIndex = SortedRemovedIndices.Num()-1; RemovalIndex >= 0; --RemovalIndex)
+			{
+				if (ExemplarIndex[Index] == SortedRemovedIndices[RemovalIndex])
+				{
+					ExemplarIndex[Index] = INDEX_NONE;
+					break;
+				}
+				else if (ExemplarIndex[Index] > SortedRemovedIndices[RemovalIndex])
+				{
+					ExemplarIndex[Index] -= (RemovalIndex + 1);
+					break;
+				}
+			}
+		}
+	}
+}
+
+
 // Input assumes that each face has a materialID that corresponds with a render material
 // This will rebuild all mesh sections
 void FGeometryCollection::ReindexMaterials()
@@ -896,7 +944,7 @@ void FGeometryCollection::Serialize(Chaos::FChaosArchive& Ar)
 		// Version 5 introduced accurate SimulationType tagging
 		if (Version < 5)
 		{
-			UE_LOG(FGeometryCollectionLogging, Warning, TEXT("GeometryCollection is has inaccurate simulation type tags. Updating tags based on transform topology."));
+			UE_LOG(FGeometryCollectionLogging, Warning, TEXT("GeometryCollection has inaccurate simulation type tags. Updating tags based on transform topology."));
 			TManagedArray<bool>* SimulatableParticles = FindAttribute<bool>(FGeometryCollection::SimulatableParticlesAttribute, FTransformCollection::TransformGroup);
 			TArray<bool> RigidChildren; RigidChildren.Init(false,NumElements(FTransformCollection::TransformGroup));
 			const TArray<int32> RecursiveOrder = GeometryCollectionAlgo::ComputeRecursiveOrder(*this);
@@ -939,6 +987,16 @@ void FGeometryCollection::Serialize(Chaos::FChaosArchive& Ar)
 			Version = 5;
 		}
 
+		// Version 6 introduced the Exemplar Index array
+		if (Version < 6)
+		{
+			UE_LOG(FGeometryCollectionLogging, Warning, TEXT("GeometryCollection has uninitialized Exemplar Indices. Setting all indices to NONE."));
+			ExemplarIndex.Fill(INDEX_NONE);
+
+			// Structure is conditioned, now considered up to date.
+			Version = 6;
+		}
+
 	}
 }
 
diff --git a/Engine/Source/Runtime/Experimental/Chaos/Private/GeometryCollection/ManagedArrayCollection.cpp b/Engine/Source/Runtime/Experimental/Chaos/Private/GeometryCollection/ManagedArrayCollection.cpp
--- a/Engine/Source/Runtime/Experimental/Chaos/Private/GeometryCollection/ManagedArrayCollection.cpp
+++ b/Engine/Source/Runtime/Experimental/Chaos/Private/GeometryCollection/ManagedArrayCollection.cpp
@@ -403,7 +403,7 @@ FString FManagedArrayCollection::ToString() const
 
 void FManagedArrayCollection::Serialize(Chaos::FChaosArchive& Ar)
 {
-	if (Ar.IsSaving()) Version = 5;
+	if (Ar.IsSaving()) Version = 6;
 	Ar << Version;
 
 	if (Ar.IsLoading())
diff --git a/Engine/Source/Runtime/Experimental/Chaos/Public/GeometryCollection/GeometryCollection.h b/Engine/Source/Runtime/Experimental/Chaos/Public/GeometryCollection/GeometryCollection.h
--- a/Engine/Source/Runtime/Experimental/Chaos/Public/GeometryCollection/GeometryCollection.h
+++ b/Engine/Source/Runtime/Experimental/Chaos/Public/GeometryCollection/GeometryCollection.h
@@ -141,6 +141,16 @@ public:
 	*/
 	int32 AppendGeometry(const FGeometryCollection & GeometryCollection, int32 MaterialIDOffset = 0, bool ReindexAllMaterials = true, const FTransform& TransformRoot = FTransform::Identity);
 
+	/**
+	* Append single embedded geometry. Returns true if the operation succeeds.
+	*/
+	bool AppendEmbeddedInstance(int32 InExemplarIndex, int32 InParentIndex, const FTransform& InTransform = FTransform::Identity);
+
+	/**
+	 * Reindex exemplar indices to reflect removed exemplars.
+	 */
+	void ReindexExemplarIndices(TArray<int32>& SortedRemovedIndices);
+
 	/**
 	* Remove Geometry and update dependent elements
 	*/
@@ -234,6 +244,7 @@ public:
 	TManagedArray<int32>        SimulationType;
 	TManagedArray<int32>        StatusFlags;
 	TManagedArray<int32>		InitialDynamicState;
+	TManagedArray<int32>		ExemplarIndex;
 
 	// Vertices Group
 	TManagedArray<FVector>		Vertex;
diff --git a/Engine/Source/Runtime/Experimental/GeometryCollectionEngine/Private/GeometryCollection/GeometryCollectionComponent.cpp b/Engine/Source/Runtime/Experimental/GeometryCollectionEngine/Private/GeometryCollection/GeometryCollectionComponent.cpp
--- a/Engine/Source/Runtime/Experimental/GeometryCollectionEngine/Private/GeometryCollection/GeometryCollectionComponent.cpp
+++ b/Engine/Source/Runtime/Experimental/GeometryCollectionEngine/Private/GeometryCollection/GeometryCollectionComponent.cpp
@@ -27,6 +27,7 @@
 #include "Net/UnrealNetwork.h"
 #include "Net/Core/PushModel/PushModel.h"
 #include "PhysicalMaterials/PhysicalMaterial.h"
+#include "Engine/InstancedStaticMesh.h"
 
 #if WITH_EDITOR
 #include "AssetToolsModule.h"
@@ -181,6 +182,8 @@ FAutoConsoleVariableRef CVarGeometryCollectionNanite(
 float GGeometryCollectionNavigationSizeThreshold = 20.0f;
 FAutoConsoleVariableRef CVarGeometryCollectionNavigationSizeThreshold(TEXT("p.GeometryCollectionNavigationSizeThreshold"), GGeometryCollectionNavigationSizeThreshold, TEXT("Size in CM used as a threshold for whether a geometry in the collection is collected and exported for navigation purposes. Measured as the diagonal of the leaf node bounds."));
 
+
+
 FGeomComponentCacheParameters::FGeomComponentCacheParameters()
 	: CacheMode(EGeometryCollectionCacheType::None)
 	, TargetCache(nullptr)
@@ -274,6 +277,7 @@ UGeometryCollectionComponent::UGeometryCollectionComponent(const FObjectInitiali
 	bHasCustomNavigableGeometry = EHasCustomNavigableGeometry::Yes;
 
 	bWantsInitializeComponent = true;
+
 }
 
 Chaos::FPhysicsSolver* GetSolver(const UGeometryCollectionComponent& GeometryCollectionComponent)
@@ -740,6 +744,53 @@ UPhysicalMaterial* UGeometryCollectionComponent::GetPhysicalMaterial() const
 	return PhysMatToUse;
 }
 
+void UGeometryCollectionComponent::RefreshEmbeddedGeometry()
+{
+	const TManagedArray<int32>& ExemplarIndexArray = GetExemplarIndexArray();
+	const int32 TransformCount = GlobalMatrices.Num();
+	
+	const int32 ExemplarCount = EmbeddedGeometryComponents.Num();
+	for (int32 ExemplarIndex = 0; ExemplarIndex < ExemplarCount; ++ExemplarIndex)
+	{		
+		TArray<FTransform> InstanceTransforms;
+		InstanceTransforms.Reserve(TransformCount); // Allocate for worst case
+
+		// Construct instance transforms for this exemplar
+		for (int32 Idx = 0; Idx < TransformCount; ++Idx)
+		{
+			if (ExemplarIndexArray[Idx] == ExemplarIndex)
+			{
+				InstanceTransforms.Add(FTransform(GlobalMatrices[Idx]));
+			}
+		}
+
+		if (EmbeddedGeometryComponents[ExemplarIndex])
+		{
+			const int32 InstanceCount = EmbeddedGeometryComponents[ExemplarIndex]->GetInstanceCount();
+
+			// If the number of instances has changed, we rebuild the structure.
+			if (InstanceCount != InstanceTransforms.Num())
+			{
+				EmbeddedGeometryComponents[ExemplarIndex]->ClearInstances();
+				EmbeddedGeometryComponents[ExemplarIndex]->PreAllocateInstancesMemory(InstanceTransforms.Num());
+				for (const FTransform& InstanceTransform : InstanceTransforms)
+				{
+					EmbeddedGeometryComponents[ExemplarIndex]->AddInstance(InstanceTransform);
+				}
+				EmbeddedGeometryComponents[ExemplarIndex]->MarkRenderStateDirty();
+			}
+			else
+			{
+				// #todo (bmiller) When ISMC has been changed to be able to update transforms in place, we need to switch this function call over.
+
+				EmbeddedGeometryComponents[ExemplarIndex]->BatchUpdateInstancesTransforms(0, InstanceTransforms, false, true, false);
+
+				// EmbeddedGeometryComponents[ExemplarIndex]->UpdateKinematicTransforms(InstanceTransforms);
+			}	
+		}
+	}
+}
+
 void UGeometryCollectionComponent::InitializeComponent()
 {
 	Super::InitializeComponent();
@@ -1630,11 +1681,13 @@ void UGeometryCollectionComponent::TickComponent(float DeltaTime, enum ELevelTic
 		{
 			if(RestCollection->HasVisibleGeometry() || DynamicCollection->IsDirty())
 			{
+				// #todo review: When we've made changes to ISMC, we need to move this function call to SetRenderDynamicData_Concurrent
+				RefreshEmbeddedGeometry();
+				
 				MarkRenderTransformDirty();
 				MarkRenderDynamicDataDirty();
 				bRenderStateDirty = false;
-				//DynamicCollection->MakeClean(); clean?
-
+				
 				const UWorld* MyWorld = GetWorld();
 				if (MyWorld && MyWorld->IsGameWorld())
 				{
@@ -1668,6 +1721,8 @@ void UGeometryCollectionComponent::OnRegister()
 
 	SetIsReplicated(bEnableReplication);
 
+	InitializeEmbeddedGeometry();
+
 	Super::OnRegister();
 }
 
@@ -2080,6 +2135,10 @@ void UGeometryCollectionComponent::SendRenderDynamicData_Concurrent()
 		}
 		else
 		{
+			// #todo (bmiller) Once ISMC changes have been complete, this is the best place to call this method
+			// but we can't currently because it's an inappropraite place to call MarkRenderStateDirty on the ISMC.
+			// RefreshEmbeddedGeometry();
+
 			// Enqueue command to send to render thread
 			if (SceneProxy->IsNaniteMesh())
 			{
@@ -2107,6 +2166,8 @@ void UGeometryCollectionComponent::SendRenderDynamicData_Concurrent()
 						}
 					}
 				);
+
+				
 			}
 		}
 
@@ -2128,6 +2189,11 @@ void UGeometryCollectionComponent::SetRestCollection(const UGeometryCollection*
 		CalculateGlobalMatrices();
 		CalculateLocalBounds();
 
+		if (!IsEmbeddedGeometryValid())
+		{
+			InitializeEmbeddedGeometry();
+		}
+		
 		//ResetDynamicCollection();
 	}
 }
@@ -2665,7 +2731,7 @@ void UGeometryCollectionComponent::CalculateGlobalMatrices()
 	{
 		// Just calc from results
 		GlobalMatrices.Empty();
-		GlobalMatrices.Append(Results->GlobalTransforms);		
+		GlobalMatrices.Append(Results->GlobalTransforms);	
 	}
 	else
 	{
@@ -2729,7 +2795,8 @@ void UGeometryCollectionComponent::SwitchRenderModels(const AActor* Actor)
 
 		if (UStaticMeshComponent* StaticMeshComp = Cast<UStaticMeshComponent>(PrimitiveComponent))
 		{
-			StaticMeshComp->SetVisibility(false);
+			// unhacked.
+			//StaticMeshComp->SetVisibility(false);
 		}
 		else if (UGeometryCollectionComponent* GeometryCollectionComponent = Cast<UGeometryCollectionComponent>(PrimitiveComponent))
 		{
@@ -2784,3 +2851,83 @@ void UGeometryCollectionComponent::EnableTransformSelectionMode(bool bEnable)
 	bIsTransformSelectionModeEnabled = bEnable;
 }
 #endif  // #if GEOMETRYCOLLECTION_EDITOR_SELECTION
+
+bool UGeometryCollectionComponent::IsEmbeddedGeometryValid() const
+{
+	// Check that the array of ISMCs that implement embedded geometry matches RestCollection Exemplar array.
+	if (!RestCollection)
+	{
+		return false;
+	}
+
+	if (RestCollection->EmbeddedGeometryExemplar.Num() != EmbeddedGeometryComponents.Num())
+	{
+		return false;
+	}
+
+	for (int32 Idx = 0; Idx < EmbeddedGeometryComponents.Num(); ++Idx)
+	{
+		UStaticMesh* ExemplarStaticMesh = Cast<UStaticMesh>(RestCollection->EmbeddedGeometryExemplar[Idx].StaticMeshExemplar.TryLoad());
+		if (!ExemplarStaticMesh)
+		{
+			return false;
+		}
+
+		if (ExemplarStaticMesh != EmbeddedGeometryComponents[Idx]->GetStaticMesh())
+		{
+			return false;
+		}
+	}
+
+	return true;
+}
+
+void UGeometryCollectionComponent::ClearEmbeddedGeometry()
+{
+	for (UInstancedStaticMeshComponent* EmbeddedGeometryComponent : EmbeddedGeometryComponents)
+	{
+		EmbeddedGeometryComponent->ClearInstances();
+		//EmbeddedGeometryComponent->UnregisterComponent();
+		EmbeddedGeometryComponent->DestroyComponent();
+	}
+
+	EmbeddedGeometryComponents.Empty();
+}
+
+void UGeometryCollectionComponent::InitializeEmbeddedGeometry()
+{
+	if (RestCollection)
+	{
+		ClearEmbeddedGeometry();
+
+		AActor* ActorOwner = GetOwner();
+		check(ActorOwner);
+
+		// Construct an InstancedStaticMeshComponent for each exemplar
+		for (const FGeometryCollectionEmbeddedExemplar& Exemplar : RestCollection->EmbeddedGeometryExemplar)
+		{
+			if (UStaticMesh* ExemplarStaticMesh = Cast<UStaticMesh>(Exemplar.StaticMeshExemplar.TryLoad()))
+			{
+				if (UInstancedStaticMeshComponent* ISMC = NewObject<UInstancedStaticMeshComponent>(ActorOwner))
+				{
+					ISMC->SetStaticMesh(ExemplarStaticMesh);
+					ISMC->SetCullDistances(Exemplar.StartCullDistance, Exemplar.EndCullDistance);
+					ISMC->SetCanEverAffectNavigation(false);
+					ISMC->SetCollisionProfileName(UCollisionProfile::NoCollision_ProfileName);
+					ISMC->SetCastShadow(false);
+					ISMC->SetMobility(EComponentMobility::Stationary);
+
+					ISMC->SetupAttachment(ActorOwner->GetRootComponent());
+					ISMC->RegisterComponent();
+					ActorOwner->AddInstanceComponent(ISMC);
+
+					EmbeddedGeometryComponents.Add(ISMC);
+				}
+			}
+		}
+
+		CalculateGlobalMatrices();
+		RefreshEmbeddedGeometry();
+
+	}
+}
diff --git a/Engine/Source/Runtime/Experimental/GeometryCollectionEngine/Private/GeometryCollection/GeometryCollectionObject.cpp b/Engine/Source/Runtime/Experimental/GeometryCollectionEngine/Private/GeometryCollection/GeometryCollectionObject.cpp
--- a/Engine/Source/Runtime/Experimental/GeometryCollectionEngine/Private/GeometryCollection/GeometryCollectionObject.cpp
+++ b/Engine/Source/Runtime/Experimental/GeometryCollectionEngine/Private/GeometryCollection/GeometryCollectionObject.cpp
@@ -721,6 +721,33 @@ bool UGeometryCollection::IsSimulationDataDirty() const
 }
 #endif
 
+int32 UGeometryCollection::AttachEmbeddedGeometryExemplar(const UStaticMesh* Exemplar)
+{
+	FSoftObjectPath NewExemplarPath(Exemplar);
+	
+	// Check first if the exemplar is already attached
+	for (int32 ExemplarIndex = 0; ExemplarIndex < EmbeddedGeometryExemplar.Num(); ++ExemplarIndex)
+	{
+		if (NewExemplarPath == EmbeddedGeometryExemplar[ExemplarIndex].StaticMeshExemplar)
+		{
+			return ExemplarIndex;
+		}
+	}
+
+	return EmbeddedGeometryExemplar.Emplace( NewExemplarPath );
+}
+
+void UGeometryCollection::RemoveExemplars(const TArray<int32>& SortedRemovalIndices)
+{
+	if (SortedRemovalIndices.Num() > 0)
+	{
+		for (int32 Index = SortedRemovalIndices.Num() - 1; Index >= 0; --Index)
+		{
+			EmbeddedGeometryExemplar.RemoveAt(Index);
+		}
+	}
+}
+
 FGuid UGeometryCollection::GetIdGuid() const
 {
 	return PersistentGuid;
diff --git a/Engine/Source/Runtime/Experimental/GeometryCollectionEngine/Public/GeometryCollection/GeometryCollectionComponent.h b/Engine/Source/Runtime/Experimental/GeometryCollectionEngine/Public/GeometryCollection/GeometryCollectionComponent.h
--- a/Engine/Source/Runtime/Experimental/GeometryCollectionEngine/Public/GeometryCollection/GeometryCollectionComponent.h
+++ b/Engine/Source/Runtime/Experimental/GeometryCollectionEngine/Public/GeometryCollection/GeometryCollectionComponent.h
@@ -32,6 +32,8 @@ class UBoxComponent;
 class UGeometryCollectionCache;
 class UChaosPhysicalMaterial;
 class AChaosSolverActor;
+struct FGeometryCollectionEmbeddedExemplar;
+class UInstancedStaticMeshComponent;
 
 DECLARE_DYNAMIC_MULTICAST_DELEGATE_OneParam(FOnChaosBreakEvent, const FChaosBreakEvent&, BreakEvent);
 
@@ -420,6 +422,7 @@ public:
 	COPY_ON_WRITE_ATTRIBUTE(int32, SimulationType, FTransformCollection::TransformGroup)
 	COPY_ON_WRITE_ATTRIBUTE(int32, TransformToGeometryIndex, FTransformCollection::TransformGroup)
 	COPY_ON_WRITE_ATTRIBUTE(int32, StatusFlags, FTransformCollection::TransformGroup)
+	COPY_ON_WRITE_ATTRIBUTE(int32, ExemplarIndex, FTransformCollection::TransformGroup)
 
 
 	UPROPERTY(EditAnywhere, NoClear, BlueprintReadOnly, Category = "ChaosPhysics")
@@ -624,6 +627,12 @@ public:
 	/** Gets the physical material to use for this geometry collection, taking into account instance overrides and render materials */
 	UPhysicalMaterial* GetPhysicalMaterial() const;
 
+	/** Update component structure to reflect any changes to the embedded geometry */
+	void InitializeEmbeddedGeometry();
+	
+	/** Update instanced static mesh components to reflect internal embedded geometry state. */
+	void RefreshEmbeddedGeometry();
+
 public:
 	UPROPERTY(BlueprintAssignable, Category = "Collision")
 	FOnChaosPhysicsCollision OnChaosPhysicsCollision;
@@ -635,6 +644,7 @@ public:
 	virtual void DispatchChaosPhysicsCollisionBlueprintEvents(const FChaosPhysicsCollisionInfo& CollisionInfo) override;
 
 protected:
+	
 	/** Call SetNotifyBreaks to set this at runtime. */
 	UPROPERTY(EditAnywhere, BlueprintReadOnly, Category = "ChaosPhysics|General")
 	bool bNotifyBreaks;
@@ -764,4 +774,12 @@ private:
 #if GEOMETRYCOLLECTION_EDITOR_SELECTION
 	bool bIsTransformSelectionModeEnabled;
 #endif  // #if GEOMETRYCOLLECTION_EDITOR_SELECTION
+
+	/** The information of all the embedded instanced static meshes */
+	UPROPERTY()
+	TArray<UInstancedStaticMeshComponent*> EmbeddedGeometryComponents;
+
+	bool IsEmbeddedGeometryValid() const;
+	void ClearEmbeddedGeometry();
+
 };
diff --git a/Engine/Source/Runtime/Experimental/GeometryCollectionEngine/Public/GeometryCollection/GeometryCollectionObject.h b/Engine/Source/Runtime/Experimental/GeometryCollectionEngine/Public/GeometryCollection/GeometryCollectionObject.h
--- a/Engine/Source/Runtime/Experimental/GeometryCollectionEngine/Public/GeometryCollection/GeometryCollectionObject.h
+++ b/Engine/Source/Runtime/Experimental/GeometryCollectionEngine/Public/GeometryCollection/GeometryCollectionObject.h
@@ -9,6 +9,7 @@
 #include "GeometryCollection/ManagedArray.h"
 #include "GeometryCollection/GeometryCollectionSimulationTypes.h"
 #include "Chaos/ChaosSolverActor.h"
+
 #include "GeometryCollectionObject.generated.h"
 
 class UMaterialInterface;
@@ -33,6 +34,32 @@ struct GEOMETRYCOLLECTIONENGINE_API FGeometryCollectionSource
 	TArray<TObjectPtr<UMaterialInterface>> SourceMaterial;
 };
 
+USTRUCT(BlueprintType)
+struct GEOMETRYCOLLECTIONENGINE_API FGeometryCollectionEmbeddedExemplar
+{
+	GENERATED_BODY()
+
+	FGeometryCollectionEmbeddedExemplar() {};
+	FGeometryCollectionEmbeddedExemplar(FSoftObjectPath NewExemplar)
+		: StaticMeshExemplar(NewExemplar)
+		, StartCullDistance(0.0f)
+		, EndCullDistance(0.0f)
+		, InstanceCount(0)
+	{ }
+
+	UPROPERTY(EditAnywhere, Category = "EmbeddedExemplar", meta = (AllowedClasses = "StaticMesh"))
+	FSoftObjectPath StaticMeshExemplar;
+
+	UPROPERTY(EditAnywhere, Category = "EmbeddedExemplar")
+	float StartCullDistance;
+
+	UPROPERTY(EditAnywhere, Category = "EmbeddedExemplar")
+	float EndCullDistance;
+
+	UPROPERTY(VisibleAnywhere, Category = "EmbeddedExemplar")
+	int32 InstanceCount;
+};
+
 
 USTRUCT()
 struct GEOMETRYCOLLECTIONENGINE_API FGeometryCollectionSizeSpecificData
@@ -199,6 +226,12 @@ public:
 	/** Check to see if Simulation Data requires regeneration */
 	bool IsSimulationDataDirty() const;
 
+	/** Attach a Static Mesh exemplar for embedded geometry, if that mesh has not already been attached. Return the exemplar index. */
+	int32 AttachEmbeddedGeometryExemplar(const UStaticMesh* Exemplar);
+
+	/** Remove embedded geometry exemplars with indices matching the sorted removal list. */
+	void RemoveExemplars(const TArray<int32>& SortedRemovalIndices);
+
 #if WITH_EDITOR
 	/** If this flag is set, we only regenerate simulation data when requested via CreateSimulationData() */
 	bool bManualDataCreate;
@@ -249,6 +282,10 @@ public:
 	UPROPERTY(EditAnywhere, Category = "Materials")
 	TArray<TObjectPtr<UMaterialInterface>> Materials;
 
+	/** References for embedded geometry generation */
+	UPROPERTY(EditAnywhere, Category = "EmbeddedGeometry")
+	TArray<FGeometryCollectionEmbeddedExemplar> EmbeddedGeometryExemplar;
+
 	/**
 	* Enable support for Nanite.
 	*/
