make fracture editor less likely to have invalid selection, and less likely to break if it does
 - check that parents are valid (i.e. it was not parent of top level) before adding to selection
 - make the "Sanitize" function for bone selections also remove invalid bones
 - ensure() rather than crash on invalid bones

this fixes a crash if you try to flatten a geometry collection w/ only one bone
also makes undo less likely to crash, because undo doesn't currently update the selection so often leaves non-existant bones selected

#rnx

diff --git a/Engine/Plugins/Experimental/ChaosEditor/Source/FractureEditor/Private/FractureToolContext.cpp b/Engine/Plugins/Experimental/ChaosEditor/Source/FractureEditor/Private/FractureToolContext.cpp
--- a/Engine/Plugins/Experimental/ChaosEditor/Source/FractureEditor/Private/FractureToolContext.cpp
+++ b/Engine/Plugins/Experimental/ChaosEditor/Source/FractureEditor/Private/FractureToolContext.cpp
@@ -26,7 +26,7 @@ void FFractureToolContext::Sanitize()
 {
 	// Ensure that children of a selected node are not also selected.
 	SelectedBones.RemoveAll([this](int32 Index) {
-		return HasSelectedAncestor(Index);
+		return !IsValidBone(Index) || HasSelectedAncestor(Index);
 		});
 
 	SelectedBones.Sort();
@@ -156,13 +156,18 @@ void FFractureToolContext::ConvertSelectionToClusterNodes()
 	TArray<int32> AddedClusterSelections;
 	for (int32 Index : SelectedBones)
 	{
+		int32 AddParent = FGeometryCollection::Invalid;
 		if (SimulationType[Index] == FGeometryCollection::ESimulationTypes::FST_Rigid)
 		{
-			AddedClusterSelections.AddUnique(Parents[Index]);
+			AddParent = Parents[Index];
 		}
 		else if (SimulationType[Index] == FGeometryCollection::ESimulationTypes::FST_None)
 		{
-			AddedClusterSelections.AddUnique(Parents[Parents[Index]]);
+			AddParent = Parents[Parents[Index]];
+		}
+		if (AddParent != FGeometryCollection::Invalid)
+		{
+			AddedClusterSelections.AddUnique(AddParent);
 		}
 	}
 	SelectedBones.Append(AddedClusterSelections);
@@ -175,6 +180,11 @@ bool FFractureToolContext::HasSelectedAncestor(int32 Index) const
 {
 	const TManagedArray<int32>& Parents = GeometryCollection->GetAttribute<int32>("Parent", FGeometryCollection::TransformGroup);
 
+	if (!ensureMsgf(Index >= 0 && Index < Parents.Num(), TEXT("Invalid index in selection: %d"), Index))
+	{
+		return false;
+	}
+
 	int32 CurrIndex = Index;
 	while (Parents[CurrIndex] != INDEX_NONE)
 	{
diff --git a/Engine/Plugins/Experimental/ChaosEditor/Source/FractureEditor/Private/FractureToolContext.h b/Engine/Plugins/Experimental/ChaosEditor/Source/FractureEditor/Private/FractureToolContext.h
--- a/Engine/Plugins/Experimental/ChaosEditor/Source/FractureEditor/Private/FractureToolContext.h
+++ b/Engine/Plugins/Experimental/ChaosEditor/Source/FractureEditor/Private/FractureToolContext.h
@@ -46,6 +46,10 @@ public:
 
 private:
 	bool HasSelectedAncestor(int32 Index) const;
+	bool IsValidBone(int32 Index) const
+	{
+		return Index >= 0 && Index < GeometryCollection->Parent.Num();
+	}
 	void ConvertSelectionToLeafNodes(int32 Index, TArray<int32>& LeafSelection);
 	void ConvertSelectionToRigidNodes(int32 Index, TArray<int32>& RigidSelection);
 	void ConvertSelectionToEmbeddedGeometryNodes(int32 Index, TArray<int32>& EmbeddedSelection);
diff --git a/Engine/Source/Runtime/Experimental/Chaos/Private/GeometryCollection/GeometryCollectionClusteringUtility.cpp b/Engine/Source/Runtime/Experimental/Chaos/Private/GeometryCollection/GeometryCollectionClusteringUtility.cpp
--- a/Engine/Source/Runtime/Experimental/Chaos/Private/GeometryCollection/GeometryCollectionClusteringUtility.cpp
+++ b/Engine/Source/Runtime/Experimental/Chaos/Private/GeometryCollection/GeometryCollectionClusteringUtility.cpp
@@ -257,7 +257,11 @@ void FGeometryCollectionClusteringUtility::ClusterBonesUnderExistingNode(FGeomet
 			// determine original parents of moved nodes so we can update their childrens names
 			for (int32 SourceElement : SourceElementsIn)
 			{
-				ParentsToUpdateNames.AddUnique(Parents[SourceElement]);
+				int32 Parent = Parents[SourceElement];
+				if (Parent != FGeometryCollection::Invalid)
+				{
+					ParentsToUpdateNames.AddUnique(Parent);
+				}
 			}
 
 			//ResetSliderTransforms(ExplodedTransforms, Transforms);
