Fracture Mode
   -- Implemented Delete Branch command.
   -- Fixed bug in Exploded Value that created problems when ISMCs are present on the GeometryCollectionActor
   -- Fixed but in can execute which was failing when an ISMC is attached to the actor
   -- Reset correct eliminates exemplars

diff --git a/Engine/Plugins/Experimental/ChaosEditor/Source/FractureEditor/Private/FractureEditorCommands.h b/Engine/Plugins/Experimental/ChaosEditor/Source/FractureEditor/Private/FractureEditorCommands.h
--- a/Engine/Plugins/Experimental/ChaosEditor/Source/FractureEditor/Private/FractureEditorCommands.h
+++ b/Engine/Plugins/Experimental/ChaosEditor/Source/FractureEditor/Private/FractureEditorCommands.h
@@ -34,6 +34,9 @@ class FFractureEditorCommands : public TCommands<FFractureEditorCommands>
 		TSharedPtr< FUICommandInfo > Uncluster;
 		TSharedPtr< FUICommandInfo > Flatten;
 		TSharedPtr< FUICommandInfo > MoveUp;
+
+		// Edit Commands
+		TSharedPtr< FUICommandInfo > DeleteBranch;
 		
 		// Generate Commands
 		TSharedPtr< FUICommandInfo > GenerateAsset;
@@ -42,7 +45,6 @@ class FFractureEditorCommands : public TCommands<FFractureEditorCommands>
 		// Embed Commands
 		TSharedPtr< FUICommandInfo > AddEmbeddedGeometry;
 		TSharedPtr< FUICommandInfo > AutoEmbedGeometry;
-		TSharedPtr< FUICommandInfo > DeleteEmbeddedGeometry;
 		
 		// Fracture Commands
 		TSharedPtr< FUICommandInfo > Uniform;
diff --git a/Engine/Plugins/Experimental/ChaosEditor/Source/FractureEditor/Private/FractureEditorModeToolkit.cpp b/Engine/Plugins/Experimental/ChaosEditor/Source/FractureEditor/Private/FractureEditorModeToolkit.cpp
--- a/Engine/Plugins/Experimental/ChaosEditor/Source/FractureEditor/Private/FractureEditorModeToolkit.cpp
+++ b/Engine/Plugins/Experimental/ChaosEditor/Source/FractureEditor/Private/FractureEditorModeToolkit.cpp
@@ -415,6 +415,10 @@ void FFractureEditorModeToolkit::Init(const TSharedPtr<IToolkitHost>& InitToolki
 		]
 	];
 
+
+
+
+
 	// Bind Chaos Commands;
 	BindCommands();
 
@@ -459,7 +463,7 @@ void FFractureEditorModeToolkit::OnObjectPostEditChange( UObject* Object, FPrope
 	}
 }
 
-const TArray<FName> FFractureEditorModeToolkit::PaletteNames = { FName(TEXT("Generate")), FName(TEXT("Select")), FName(TEXT("Fracture")), FName(TEXT("Cluster")), FName(TEXT("Embed")), FName(TEXT("Properties")) };
+const TArray<FName> FFractureEditorModeToolkit::PaletteNames = { FName(TEXT("Generate")), FName(TEXT("Select")), FName(TEXT("Fracture")), FName(TEXT("Edit")), FName(TEXT("Cluster")), FName(TEXT("Embed")), FName(TEXT("Properties")) };
 
 FText FFractureEditorModeToolkit::GetToolPaletteDisplayName(FName Palette) const
 { 
@@ -508,6 +512,10 @@ void FFractureEditorModeToolkit::BuildToolPalette(FName PaletteIndex, class FToo
 		ToolbarBuilder.AddToolBarButton(Commands.Brick);
 		ToolbarBuilder.AddToolBarButton(Commands.Resample);
 	}
+	else if (PaletteIndex == TEXT("Edit"))
+	{
+		ToolbarBuilder.AddToolBarButton(Commands.DeleteBranch);
+	}
 	else if (PaletteIndex == TEXT("Cluster"))
 	{
 		ToolbarBuilder.AddToolBarButton(Commands.AutoCluster);
@@ -521,7 +529,6 @@ void FFractureEditorModeToolkit::BuildToolPalette(FName PaletteIndex, class FToo
 	{
 		ToolbarBuilder.AddToolBarButton(Commands.AddEmbeddedGeometry);
 		ToolbarBuilder.AddToolBarButton(Commands.AutoEmbedGeometry);
-		ToolbarBuilder.AddToolBarButton(Commands.DeleteEmbeddedGeometry);
 	}
 	else if (PaletteIndex == TEXT("Properties"))
 	{
@@ -704,13 +711,15 @@ void FFractureEditorModeToolkit::OnExplodedViewValueChanged()
 			AGeometryCollectionActor* GeometryCollectionActor = Cast<AGeometryCollectionActor>(Actor);
 			if(GeometryCollectionActor)
 			{
-				UGeometryCollectionComponent* GeometryCollectionComponent = Cast<UGeometryCollectionComponent>(PrimitiveComponent);
-				FGeometryCollectionEdit RestCollection = GeometryCollectionComponent->EditRestCollection();
-				UGeometryCollection* GeometryCollection = RestCollection.GetRestCollection();
+				if (UGeometryCollectionComponent* GeometryCollectionComponent = Cast<UGeometryCollectionComponent>(PrimitiveComponent))
+				{
+					FGeometryCollectionEdit RestCollection = GeometryCollectionComponent->EditRestCollection();
+					UGeometryCollection* GeometryCollection = RestCollection.GetRestCollection();
 
-				UpdateExplodedVectors(GeometryCollectionComponent);
+					UpdateExplodedVectors(GeometryCollectionComponent);
 
-				GeometryCollectionComponent->MarkRenderStateDirty();
+					GeometryCollectionComponent->MarkRenderStateDirty();
+				}	
 			}
 		}
 	}
diff --git a/Engine/Plugins/Experimental/ChaosEditor/Source/FractureEditor/Private/FractureToolCutter.cpp b/Engine/Plugins/Experimental/ChaosEditor/Source/FractureEditor/Private/FractureToolCutter.cpp
--- a/Engine/Plugins/Experimental/ChaosEditor/Source/FractureEditor/Private/FractureToolCutter.cpp
+++ b/Engine/Plugins/Experimental/ChaosEditor/Source/FractureEditor/Private/FractureToolCutter.cpp
@@ -29,12 +29,7 @@ bool UFractureToolCutterBase::CanExecute() const
 	{
 		return false;
 	}
-
-	if (IsStaticMeshSelected())
-	{
-		return false;
-	}
-
+	
 	return true;
 }
 
diff --git a/Engine/Plugins/Experimental/ChaosEditor/Source/FractureEditor/Private/FractureToolEditing.cpp b/Engine/Plugins/Experimental/ChaosEditor/Source/FractureEditor/Private/FractureToolEditing.cpp
new file mode 100644
--- /dev/null
+++ b/Engine/Plugins/Experimental/ChaosEditor/Source/FractureEditor/Private/FractureToolEditing.cpp
@@ -0,0 +1,98 @@
+// Copyright Epic Games, Inc. All Rights Reserved.
+
+#include "FractureToolEditing.h"
+
+#include "Editor.h"
+
+#include "GeometryCollection/GeometryCollection.h"
+#include "GeometryCollection/GeometryCollectionClusteringUtility.h"
+#include "FractureToolContext.h"
+
+
+#define LOCTEXT_NAMESPACE "FractureToolEditing"
+
+
+FText UFractureToolDeleteBranch::GetDisplayText() const
+{
+	return FText(NSLOCTEXT("FractureToolEditingOps", "FractureToolDeleteBranch", "Delete"));
+}
+
+FText UFractureToolDeleteBranch::GetTooltipText() const
+{
+	return FText(NSLOCTEXT("FractureToolEditingOps", "FractureToolDeleteBranchTooltip", "Delete all nodes in selected branch. Empty clusters will be eliminated."));
+}
+
+FSlateIcon UFractureToolDeleteBranch::GetToolIcon() const
+{
+	return FSlateIcon("FractureEditorStyle", "FractureEditor.DeleteBranch");
+}
+
+void UFractureToolDeleteBranch::RegisterUICommand(FFractureEditorCommands* BindingContext)
+{
+	UI_COMMAND_EXT(BindingContext, UICommandInfo, "DeleteBranch", "Prune", "Delete all nodes in selected branch. Empty clusters will be eliminated.", EUserInterfaceActionType::Button, FInputChord());
+	BindingContext->DeleteBranch = UICommandInfo;
+}
+
+void UFractureToolDeleteBranch::Execute(TWeakPtr<FFractureEditorModeToolkit> InToolkit)
+{
+	if (InToolkit.IsValid())
+	{
+		FFractureEditorModeToolkit* Toolkit = InToolkit.Pin().Get();
+
+		TArray<FFractureToolContext> Contexts = GetFractureToolContexts();
+
+		for (FFractureToolContext& Context : Contexts)
+		{
+			FGeometryCollection* GeometryCollection = Context.GetGeometryCollection().Get();
+			UGeometryCollection* FracturedGeometryCollection = Context.GetFracturedGeometryCollection();
+
+			const TManagedArray<int32>& ExemplarIndex = GeometryCollection->ExemplarIndex;
+			const TManagedArray<TSet<int32>>& Children = GeometryCollection->Children;
+			
+			// Removing the root node amounts to full deletion -- we don't allow this here.
+			Context.RemoveRootNodes();
+			Context.Sanitize();
+
+			TArray<int32> NodesForDeletion;
+
+			for (int32 Select : Context.GetSelection())
+			{
+				FGeometryCollectionClusteringUtility::RecursiveAddAllChildren(Children, Select, NodesForDeletion);
+			}
+
+			// Clean up any embedded geometry removal
+			TArray<int32> UninstancedExemplars;
+			UninstancedExemplars.Reserve(NodesForDeletion.Num());
+			for (int32 DeleteNode : NodesForDeletion)
+			{
+				if (ExemplarIndex[DeleteNode] > INDEX_NONE)
+				{
+					if ((--FracturedGeometryCollection->EmbeddedGeometryExemplar[ExemplarIndex[DeleteNode]].InstanceCount) < 1)
+					{
+						UE_LOG(LogFractureTool, Warning, TEXT("Exemplar Index %d is empty. Removing Exemplar from Geometry Collection."), ExemplarIndex[DeleteNode]);
+						UninstancedExemplars.Add(ExemplarIndex[DeleteNode]);
+					}
+				}
+			}
+
+			UninstancedExemplars.Sort();
+			FracturedGeometryCollection->RemoveExemplars(UninstancedExemplars);
+			GeometryCollection->ReindexExemplarIndices(UninstancedExemplars);
+
+			NodesForDeletion.Sort();
+			GeometryCollection->RemoveElements(FGeometryCollection::TransformGroup, NodesForDeletion);
+
+			FGeometryCollectionClusteringUtility::RemoveDanglingClusters(GeometryCollection);
+
+			Context.GetGeometryCollectionComponent()->InitializeEmbeddedGeometry();
+
+			Refresh(Context, Toolkit, true);
+		}
+
+		SetOutlinerComponents(Contexts, Toolkit);
+	}
+}
+
+
+#undef LOCTEXT_NAMESPACE
+
diff --git a/Engine/Plugins/Experimental/ChaosEditor/Source/FractureEditor/Private/FractureToolEditing.h b/Engine/Plugins/Experimental/ChaosEditor/Source/FractureEditor/Private/FractureToolEditing.h
new file mode 100644
--- /dev/null
+++ b/Engine/Plugins/Experimental/ChaosEditor/Source/FractureEditor/Private/FractureToolEditing.h
@@ -0,0 +1,23 @@
+// Copyright Epic Games, Inc. All Rights Reserved.
+
+#pragma once
+
+#include "FractureTool.h"
+
+#include "FractureToolEditing.generated.h"
+
+UCLASS(DisplayName = "Delete Branch", Category = "FractureTools")
+class UFractureToolDeleteBranch : public UFractureActionTool
+{
+public:
+	GENERATED_BODY()
+
+	UFractureToolDeleteBranch(const FObjectInitializer& ObjInit) : Super(ObjInit) {}
+
+	// UFractureActionTool Interface
+	virtual FText GetDisplayText() const override;
+	virtual FText GetTooltipText() const override;
+	virtual FSlateIcon GetToolIcon() const override;
+	virtual void RegisterUICommand(FFractureEditorCommands* BindingContext) override;
+	virtual void Execute(TWeakPtr<FFractureEditorModeToolkit> InToolkit) override;
+};
diff --git a/Engine/Plugins/Experimental/ChaosEditor/Source/FractureEditor/Private/FractureToolEmbed.cpp b/Engine/Plugins/Experimental/ChaosEditor/Source/FractureEditor/Private/FractureToolEmbed.cpp
--- a/Engine/Plugins/Experimental/ChaosEditor/Source/FractureEditor/Private/FractureToolEmbed.cpp
+++ b/Engine/Plugins/Experimental/ChaosEditor/Source/FractureEditor/Private/FractureToolEmbed.cpp
@@ -288,111 +288,4 @@ TArray<UStaticMeshComponent*> UFractureToolAutoEmbedGeometry::GetSelectedStaticM
 }
 
 
-
-
-FText UFractureToolDeleteEmbeddedGeometry::GetDisplayText() const
-{
-	return FText(NSLOCTEXT("Fracture", "FractureToolDeleteEmbeddedGeometry", "Delete"));
-}
-
-FText UFractureToolDeleteEmbeddedGeometry::GetTooltipText() const
-{
-	return FText(NSLOCTEXT("Fracture", "FractureToolDeleteEmbeddedGeometryTooltip", "Delete selected embedded geometry nodes. If a cluster or rigid node is selected, all child embedded geometry nodes are deleted."));
-}
-
-FSlateIcon UFractureToolDeleteEmbeddedGeometry::GetToolIcon() const
-{
-	return FSlateIcon("FractureEditorStyle", "FractureEditor.DeleteEmbeddedGeometry");
-}
-
-void UFractureToolDeleteEmbeddedGeometry::RegisterUICommand(FFractureEditorCommands* BindingContext)
-{
-	UI_COMMAND_EXT(BindingContext, UICommandInfo, "DeleteEmbeddedGeometry", "Delete", "Delete", EUserInterfaceActionType::ToggleButton, FInputChord());
-	BindingContext->DeleteEmbeddedGeometry = UICommandInfo;
-}
-
-bool UFractureToolDeleteEmbeddedGeometry::CanExecute() const
-{
-	return IsGeometryCollectionSelected();
-}
-
-void UFractureToolDeleteEmbeddedGeometry::Execute(TWeakPtr<FFractureEditorModeToolkit> InToolkit)
-{
-	if (InToolkit.IsValid())
-	{
-		FFractureEditorModeToolkit* Toolkit = InToolkit.Pin().Get();
-
-		// Get selected static meshes
-		TArray<FFractureToolContext> Contexts = GetFractureToolContexts();
-
-		for (FFractureToolContext& Context : Contexts)
-		{
-			Context.Sanitize();
-
-			FGeometryCollectionEdit GeometryCollectionEdit = Context.GetGeometryCollectionComponent()->EditRestCollection(GeometryCollection::EEditUpdate::RestPhysicsDynamic);
-			UGeometryCollection* FracturedGeometryCollection = GeometryCollectionEdit.GetRestCollection();
-
-			FFractureToolContext::FGeometryCollectionPtr GeometryCollection = Context.GetGeometryCollection();
-			const TManagedArray<int32>& ExemplarIndex = GeometryCollection->ExemplarIndex;
-
-			TArray<int32> EmbeddedGeometryToBeRemoved;
-			const TArray<int32>& SelectedBones = Context.GetSelection();
-			for (const int32 SelectedBone : SelectedBones)
-			{
-				if (ExemplarIndex[SelectedBone] == INDEX_NONE)
-				{
-					// Select all Embedded Geometry found in the selected branch.
-					TArray<int32> LeafBones;
-					FGeometryCollectionClusteringUtility::GetLeafBones(GeometryCollection.Get(), SelectedBone, false, LeafBones);
-					for (int32 LeafBone : LeafBones)
-					{
-						if (ExemplarIndex[LeafBone] > INDEX_NONE)
-						{
-							EmbeddedGeometryToBeRemoved.Add(LeafBone);
-						}
-					}
-				}
-				else
-				{
-					// Selected bone is embedded geometry. Only delete this.
-					EmbeddedGeometryToBeRemoved.Add(SelectedBone);
-				}
-			}
-
-			TArray<int32> UninstancedExemplars;
-			UninstancedExemplars.Reserve(EmbeddedGeometryToBeRemoved.Num());
-			for (int32 EmbeddedIndex : EmbeddedGeometryToBeRemoved)
-			{
-				if ((--FracturedGeometryCollection->EmbeddedGeometryExemplar[ExemplarIndex[EmbeddedIndex]].InstanceCount) < 1)
-				{
-					UE_LOG(LogFractureTool, Warning, TEXT("Exemplar Index %d is empty. Removing Exemplar from Geometry Collection."), ExemplarIndex[EmbeddedIndex]);
-					UninstancedExemplars.Add(ExemplarIndex[EmbeddedIndex]);
-				}
-			}
-
-			EmbeddedGeometryToBeRemoved.Sort();
-			GeometryCollection->RemoveElements(FGeometryCollection::TransformGroup, EmbeddedGeometryToBeRemoved);
-
-			UninstancedExemplars.Sort();
-			FracturedGeometryCollection->RemoveExemplars(UninstancedExemplars);
-			GeometryCollection->ReindexExemplarIndices(UninstancedExemplars);
-
-			Context.GetGeometryCollectionComponent()->InitializeEmbeddedGeometry();
-			Refresh(Context, Toolkit, true);
-			FracturedGeometryCollection->MarkPackageDirty();
-		}
-
-		SetOutlinerComponents(Contexts, Toolkit);
-	}
-}
-
-
-
-
-
-
-
-
-
-
 #undef LOCTEXT_NAMESPACE
\ No newline at end of file
diff --git a/Engine/Plugins/Experimental/ChaosEditor/Source/FractureEditor/Private/FractureToolEmbed.h b/Engine/Plugins/Experimental/ChaosEditor/Source/FractureEditor/Private/FractureToolEmbed.h
--- a/Engine/Plugins/Experimental/ChaosEditor/Source/FractureEditor/Private/FractureToolEmbed.h
+++ b/Engine/Plugins/Experimental/ChaosEditor/Source/FractureEditor/Private/FractureToolEmbed.h
@@ -47,22 +47,3 @@ public:
 private:
 	static TArray<UStaticMeshComponent*>  GetSelectedStaticMeshComponents();
 };
-
-
-UCLASS()
-class UFractureToolDeleteEmbeddedGeometry : public UFractureActionTool
-{
-public:
-	GENERATED_BODY()
-
-	UFractureToolDeleteEmbeddedGeometry(const FObjectInitializer& ObjInit) : Super(ObjInit) {}
-
-	// UFractureActionTool Interface
-	virtual FText GetDisplayText() const override;
-	virtual FText GetTooltipText() const override;
-	virtual FSlateIcon GetToolIcon() const override;
-	virtual void RegisterUICommand(FFractureEditorCommands* BindingContext) override;
-	virtual void Execute(TWeakPtr<FFractureEditorModeToolkit> InToolkit) override;
-	virtual bool CanExecute() const override;
-};
-
diff --git a/Engine/Plugins/Experimental/ChaosEditor/Source/FractureEditor/Private/FractureToolGenerators.cpp b/Engine/Plugins/Experimental/ChaosEditor/Source/FractureEditor/Private/FractureToolGenerators.cpp
--- a/Engine/Plugins/Experimental/ChaosEditor/Source/FractureEditor/Private/FractureToolGenerators.cpp
+++ b/Engine/Plugins/Experimental/ChaosEditor/Source/FractureEditor/Private/FractureToolGenerators.cpp
@@ -415,9 +415,12 @@ void UFractureToolResetAsset::Execute(TWeakPtr<FFractureEditorModeToolkit> InToo
 				GeometryCollectionComponent->MarkRenderDynamicDataDirty();
 				GeometryCollectionComponent->MarkRenderStateDirty();
 			}
+
 			GeometryCollectionObject->MarkPackageDirty();
 		}
 		
+		GeometryCollectionComponent->InitializeEmbeddedGeometry();
+		
 		FScopedColorEdit EditBoneColor = GeometryCollectionComponent->EditBoneSelection();
 		EditBoneColor.ResetBoneSelection();
 		EditBoneColor.ResetHighlightedBones();
diff --git a/Engine/Source/Runtime/Experimental/GeometryCollectionEngine/Private/GeometryCollection/GeometryCollectionComponent.cpp b/Engine/Source/Runtime/Experimental/GeometryCollectionEngine/Private/GeometryCollection/GeometryCollectionComponent.cpp
--- a/Engine/Source/Runtime/Experimental/GeometryCollectionEngine/Private/GeometryCollection/GeometryCollectionComponent.cpp
+++ b/Engine/Source/Runtime/Experimental/GeometryCollectionEngine/Private/GeometryCollection/GeometryCollectionComponent.cpp
@@ -2884,11 +2884,16 @@ bool UGeometryCollectionComponent::IsEmbeddedGeometryValid() const
 
 void UGeometryCollectionComponent::ClearEmbeddedGeometry()
 {
-	for (UInstancedStaticMeshComponent* EmbeddedGeometryComponent : EmbeddedGeometryComponents)
+	AActor* OwningActor = GetOwner();
+	TArray<UActorComponent*> TargetComponents = OwningActor->GetComponentsByClass(UInstancedStaticMeshComponent::StaticClass());
+	
+	for (UActorComponent* EmbeddedGeometryComponent : TargetComponents)
 	{
-		EmbeddedGeometryComponent->ClearInstances();
-		//EmbeddedGeometryComponent->UnregisterComponent();
-		EmbeddedGeometryComponent->DestroyComponent();
+		if (UInstancedStaticMeshComponent* ISMComponent = Cast<UInstancedStaticMeshComponent>(EmbeddedGeometryComponent))
+		{
+			ISMComponent->ClearInstances();
+			ISMComponent->DestroyComponent();
+		}
 	}
 
 	EmbeddedGeometryComponents.Empty();
diff --git a/Engine/Source/Runtime/Experimental/GeometryCollectionEngine/Private/GeometryCollection/GeometryCollectionObject.cpp b/Engine/Source/Runtime/Experimental/GeometryCollectionEngine/Private/GeometryCollection/GeometryCollectionObject.cpp
--- a/Engine/Source/Runtime/Experimental/GeometryCollectionEngine/Private/GeometryCollection/GeometryCollectionObject.cpp
+++ b/Engine/Source/Runtime/Experimental/GeometryCollectionEngine/Private/GeometryCollection/GeometryCollectionObject.cpp
@@ -189,6 +189,7 @@ void UGeometryCollection::Reset()
 		Modify();
 		GeometryCollection->Empty();
 		Materials.Empty();
+		EmbeddedGeometryExemplar.Empty();
 		InvalidateCollection();
 	}
 }
