Fracture Tool -- sped up SM conversion when large numbers of SMs are imported by waiting until completion to reindex materials.

#okforgitbug public,internal

diff --git a/Engine/Plugins/Experimental/ChaosEditor/Source/FractureEditor/Private/FractureToolGenerators.cpp b/Engine/Plugins/Experimental/ChaosEditor/Source/FractureEditor/Private/FractureToolGenerators.cpp
--- a/Engine/Plugins/Experimental/ChaosEditor/Source/FractureEditor/Private/FractureToolGenerators.cpp
+++ b/Engine/Plugins/Experimental/ChaosEditor/Source/FractureEditor/Private/FractureToolGenerators.cpp
@@ -203,7 +203,7 @@ AGeometryCollectionActor* UFractureToolGenerateAsset::ConvertActorsToGeometryCol
 					decltype(FGeometryCollectionSource::SourceMaterial) SourceMaterials(StaticMeshComponent->GetMaterials());
 					FracturedGeometryCollection->GeometrySource.Add({ SourceSoftObjectPath, ComponentTransform, SourceMaterials });
 
-					FGeometryCollectionConversion::AppendStaticMesh(ComponentStaticMesh, SourceMaterials, ComponentTransform, FracturedGeometryCollection, true);
+					FGeometryCollectionConversion::AppendStaticMesh(ComponentStaticMesh, SourceMaterials, ComponentTransform, FracturedGeometryCollection, false);
 				}
 			}
 		}
@@ -239,7 +239,7 @@ AGeometryCollectionActor* UFractureToolGenerateAsset::ConvertActorsToGeometryCol
 					}
 					FracturedGeometryCollection->GeometrySource.Add({ SourceSoftObjectPath, ComponentTransform, SourceMaterials });
 
-					FGeometryCollectionConversion::AppendGeometryCollection(RestCollection, GeometryCollectionComponent, ComponentTransform, FracturedGeometryCollection, true);
+					FGeometryCollectionConversion::AppendGeometryCollection(RestCollection, GeometryCollectionComponent, ComponentTransform, FracturedGeometryCollection, false);
 
 				}
 			}
