Fracture Mode -- Drag&Drop functionality in the Fracture Mode outliner.

diff --git a/Engine/Plugins/Experimental/ChaosEditor/Source/FractureEditor/Private/GeometryCollectionOutlinerDragDrop.cpp b/Engine/Plugins/Experimental/ChaosEditor/Source/FractureEditor/Private/GeometryCollectionOutlinerDragDrop.cpp
new file mode 100644
--- /dev/null
+++ b/Engine/Plugins/Experimental/ChaosEditor/Source/FractureEditor/Private/GeometryCollectionOutlinerDragDrop.cpp
@@ -0,0 +1,180 @@
+// Copyright Epic Games, Inc. All Rights Reserved.
+
+#include "GeometryCollectionOutlinerDragDrop.h"
+
+#include "GeometryCollection/GeometryCollection.h"
+#include "GeometryCollection/GeometryCollectionClusteringUtility.h"
+
+#define LOCTEXT_NAMESPACE "ChaosEditor"
+
+bool FGeometryCollectionBoneDragDrop::ValidateDrop(const FGeometryCollection* OtherGeometryCollection, int32 OtherBone, FText& MessageText)
+{
+	if (!OtherGeometryCollection)
+	{
+		return false;
+	}
+	
+	// We don't currently support transfer of geometry from one collection to another.
+	if (OtherGeometryCollection != GeometryCollection.Get())
+	{
+		MessageText = NSLOCTEXT("GeometryCollectionOutliner", "ComponentTransfer_GeometryCollectionOutliner", "Bones must remain within original component.");
+		return false;
+	}
+
+	if (BonePayload.Contains(OtherBone))
+	{
+		MessageText = NSLOCTEXT("GeometryCollectionOutliner", "Self_GeometryCollectionOutliner", "Cannot parent to self.");
+		return false;
+	}
+
+	const TManagedArray<int32>& ExemplarIndex = OtherGeometryCollection->ExemplarIndex;
+	const TManagedArray<int32>& SimulationType = OtherGeometryCollection->SimulationType;
+	if (OtherGeometryCollection->IsClustered(OtherBone))
+	{
+		if (ContainsInstance())
+		{
+			MessageText = NSLOCTEXT("GeometryCollectionOutliner", "Instances_GeometryCollectionOutliner", 
+				"Cannot parent instanced embedded geometry directly to cluster.");
+			return false;
+		}
+
+		if (ContainsInstance())
+		{
+			MessageText = NSLOCTEXT("GeometryCollectionOutliner", "EmbeddedToCluster_GeometryCollectionOutliner",
+				"Cannot parent instanced embedded geometry to cluster.");
+			return false;
+		}
+
+		if (ContainsEmbedded())
+		{
+			MessageText = NSLOCTEXT("GeometryCollectionOutliner", "EmbeddedToCluster_GeometryCollectionOutliner",
+				"Embedded geometry will convert to rigids when parented to cluster.");
+			return true;
+		}
+
+		return true;
+	}
+	else if (OtherGeometryCollection->IsRigid(OtherBone))
+	{
+		if (ContainsCluster())
+		{
+			MessageText = NSLOCTEXT("GeometryCollectionOutliner", "ClusterToRigid_GeometryCollectionOutliner",
+				"Cannot parent cluster to rigid.");
+			return false;
+		}
+
+		if (ContainsRigid())
+		{
+			MessageText = NSLOCTEXT("GeometryCollectionOutliner", "RigidToRigid_GeometryCollectionOutliner",
+				"Rigids will convert to embedded geometry when parented to rigid.");
+			return true;
+		}
+
+		return true;
+	}
+	else if (SimulationType[OtherBone] == FGeometryCollection::ESimulationTypes::FST_None)
+	{
+		if (ContainsCluster() || ContainsRigid())
+		{
+			MessageText = NSLOCTEXT("GeometryCollectionOutliner", "ToEmbedded_GeometryCollectionOutliner",
+				"Cannot parent to embedded geometry.");
+			return false;
+		}
+
+		return true;
+	}
+
+	// We disallow any unrecognized combinations
+	return false;
+}
+
+bool FGeometryCollectionBoneDragDrop::ReparentBones(const FGeometryCollection* OtherGeometryCollection, int32 OtherBone)
+{
+	FText HoverText;
+	bool bValid = ValidateDrop(OtherGeometryCollection, OtherBone, HoverText);
+	if (bValid)
+	{
+		// If we parent a rigid to another rigid, we convert the child rigid to embedded.
+		if (GeometryCollection->IsRigid(OtherBone))
+		{
+			TManagedArray<int32>& SimulationType = GeometryCollection->SimulationType;
+			for (int32 Bone : BonePayload)
+			{
+				if (GeometryCollection->IsRigid(Bone))
+				{
+					SimulationType[Bone] = FGeometryCollection::ESimulationTypes::FST_None;
+				}
+			}
+		}
+		// If we parent an embedded geometry to a cluster, it becomes a rigid body.
+		else if (GeometryCollection->IsClustered(OtherBone))
+		{
+			const TManagedArray<int32>& ExemplarIndex = OtherGeometryCollection->ExemplarIndex;
+			TManagedArray<int32>& SimulationType = GeometryCollection->SimulationType;
+			for (int32 Bone : BonePayload)
+			{
+				if ((SimulationType[Bone] == FGeometryCollection::ESimulationTypes::FST_None) && (ExemplarIndex[Bone] == INDEX_NONE))
+				{
+					SimulationType[Bone] = FGeometryCollection::ESimulationTypes::FST_Rigid;
+				}
+			}
+		}
+
+		FGeometryCollectionClusteringUtility::ClusterBonesUnderExistingNode(GeometryCollection.Get(), OtherBone, BonePayload);
+	}
+
+	return bValid;
+}
+
+bool FGeometryCollectionBoneDragDrop::ContainsCluster() const
+{
+	for (int32 Bone : BonePayload)
+	{
+		if (GeometryCollection->IsClustered(Bone))
+		{
+			return true;
+		}
+	}
+	return false;
+}
+
+bool FGeometryCollectionBoneDragDrop::ContainsRigid() const
+{
+	for (int32 Bone : BonePayload)
+	{
+		if (GeometryCollection->IsRigid(Bone))
+		{
+			return true;
+		}
+	}
+	return false;
+}
+
+bool FGeometryCollectionBoneDragDrop::ContainsEmbedded() const
+{
+	const TManagedArray<int32>& SimulationType = GeometryCollection->SimulationType;
+	for (int32 Bone : BonePayload)
+	{
+		if (SimulationType[Bone] == FGeometryCollection::ESimulationTypes::FST_None)
+		{
+			return true;
+		}
+	}
+	return false;
+}
+
+bool FGeometryCollectionBoneDragDrop::ContainsInstance() const
+{
+	const TManagedArray<int32>& ExemplarIndex = GeometryCollection->ExemplarIndex;
+	for (int32 Bone : BonePayload)
+	{
+		if (ExemplarIndex[Bone] > INDEX_NONE)
+		{
+			return true;
+		}
+	}
+	return false;
+}
+
+
+#undef LOCTEXT_NAMESPACE
\ No newline at end of file
diff --git a/Engine/Plugins/Experimental/ChaosEditor/Source/FractureEditor/Private/GeometryCollectionOutlinerDragDrop.h b/Engine/Plugins/Experimental/ChaosEditor/Source/FractureEditor/Private/GeometryCollectionOutlinerDragDrop.h
new file mode 100644
--- /dev/null
+++ b/Engine/Plugins/Experimental/ChaosEditor/Source/FractureEditor/Private/GeometryCollectionOutlinerDragDrop.h
@@ -0,0 +1,47 @@
+// Copyright Epic Games, Inc. All Rights Reserved.
+
+#pragma once
+
+#include "CoreMinimal.h"
+#include "DragAndDrop/DecoratedDragDropOp.h"
+
+#include "SGeometryCollectionOutliner.h"
+
+/**
+* This operation is a bone which can be dragged and dropped to another bone in the outliner.
+*/
+class FGeometryCollectionBoneDragDrop
+	: public FDecoratedDragDropOp
+{
+public:
+
+	DRAG_DROP_OPERATOR_TYPE(FGeometryCollectionBoneDragDrop, FDecoratedDragDropOp)
+
+
+	TArray<int32> BonePayload;
+	TSharedPtr<FGeometryCollection,ESPMode::ThreadSafe> GeometryCollection;
+
+	static TSharedRef<FGeometryCollectionBoneDragDrop> New(TSharedPtr<FGeometryCollection,ESPMode::ThreadSafe> InGeometryCollection, TArray<int32>& InBonePayload)
+	{
+		TSharedRef<FGeometryCollectionBoneDragDrop> Operation = MakeShareable(new FGeometryCollectionBoneDragDrop);
+		Operation->MouseCursor = EMouseCursor::GrabHandClosed;
+		Operation->BonePayload = InBonePayload;
+		Operation->GeometryCollection = InGeometryCollection;
+
+		Operation->Construct();
+		return Operation;
+	}
+
+	/** Return true if dropping the BonePayload onto OtherBone is a valid operation. If it's not, return false and provide reason in MessageText. */
+	bool ValidateDrop(const FGeometryCollection* OtherGeometryCollection, int32 OtherBone, FText& MessageText);
+
+	/** Reparent BonePayload to OtherBone if this is a valid operation. Return false otherwise. */
+	bool ReparentBones(const FGeometryCollection* OtherGeometryCollection, int32 OtherBone);
+
+private:
+	bool ContainsCluster() const;
+	bool ContainsRigid() const;
+	bool ContainsEmbedded() const;
+	bool ContainsInstance() const;
+	
+};
\ No newline at end of file
diff --git a/Engine/Plugins/Experimental/ChaosEditor/Source/FractureEditor/Private/SGeometryCollectionOutliner.cpp b/Engine/Plugins/Experimental/ChaosEditor/Source/FractureEditor/Private/SGeometryCollectionOutliner.cpp
--- a/Engine/Plugins/Experimental/ChaosEditor/Source/FractureEditor/Private/SGeometryCollectionOutliner.cpp
+++ b/Engine/Plugins/Experimental/ChaosEditor/Source/FractureEditor/Private/SGeometryCollectionOutliner.cpp
@@ -13,15 +13,103 @@
 #include "FractureEditorMode.h"
 #include "ScopedTransaction.h"
 #include "FractureSettings.h"
+#include "GeometryCollectionOutlinerDragDrop.h"
 
 #define LOCTEXT_NAMESPACE "ChaosEditor"
 
+void FGeometryCollectionTreeItem::OnDragLeave(const FDragDropEvent& InDragDropEvent)
+{
+	TSharedPtr<FDragDropOperation> Operation = InDragDropEvent.GetOperation();
+	if (Operation.IsValid())
+	{
+		if (Operation->IsOfType<FGeometryCollectionBoneDragDrop>())
+		{
+			TSharedPtr<FGeometryCollectionBoneDragDrop> GeometryCollectionBoneOp = InDragDropEvent.GetOperationAs<FGeometryCollectionBoneDragDrop>();
+			const FSlateBrush* Icon = FEditorStyle::GetBrush(TEXT("Graph.ConnectorFeedback.Error"));
+			GeometryCollectionBoneOp->SetToolTip(FText(), Icon);
+		}
+	}
+}
+
+void FGeometryCollectionTreeItemBone::OnDragEnter(FDragDropEvent const& InDragDropEvent)
+{
+	TSharedPtr<FDragDropOperation> Operation = InDragDropEvent.GetOperation();
+	if (Operation.IsValid())
+	{
+		if (Operation->IsOfType<FGeometryCollectionBoneDragDrop>())
+		{
+			TSharedPtr<FGeometryCollectionBoneDragDrop> GeometryCollectionBoneOp = InDragDropEvent.GetOperationAs<FGeometryCollectionBoneDragDrop>();
+			
+			UGeometryCollectionComponent* SourceComponent = GetComponent();
+			FGeometryCollectionEdit GeometryCollectionEdit = SourceComponent->EditRestCollection();
+			if (UGeometryCollection* GeometryCollectionObject = GeometryCollectionEdit.GetRestCollection())
+			{
+				FGeometryCollection* GeometryCollection = GeometryCollectionObject->GetGeometryCollection().Get();
+				FText HoverText;
+				bool bValid = GeometryCollectionBoneOp->ValidateDrop(GeometryCollection, BoneIndex, HoverText);
+				const FSlateBrush* Icon = bValid ? FEditorStyle::GetBrush(TEXT("Graph.ConnectorFeedback.OK")) : FEditorStyle::GetBrush(TEXT("Graph.ConnectorFeedback.Error"));
+				GeometryCollectionBoneOp->SetToolTip(HoverText, Icon);
+			}
+		}
+	}
+}
+
+FReply FGeometryCollectionTreeItemBone::OnDragDetected(const FGeometry& MyGeometry, const FPointerEvent& MouseEvent)
+{
+	if (MouseEvent.IsMouseButtonDown(EKeys::LeftMouseButton))
+	{
+		UGeometryCollectionComponent* SourceComponent = GetComponent();
+
+		FGeometryCollectionEdit GeometryCollectionEdit = SourceComponent->EditRestCollection();
+		if (UGeometryCollection* GeometryCollectionObject = GeometryCollectionEdit.GetRestCollection())
+		{
+			TSharedPtr<FGeometryCollection,ESPMode::ThreadSafe> GeometryCollection = GeometryCollectionObject->GetGeometryCollection();
+			TArray<int32> SelectedBones = SourceComponent->GetSelectedBones();
+
+			return FReply::Handled().BeginDragDrop(FGeometryCollectionBoneDragDrop::New(GeometryCollection, SelectedBones));
+		}		
+	}
+
+	return FReply::Unhandled();
+}
+
+FReply FGeometryCollectionTreeItemBone::OnDrop(const FDragDropEvent& DragDropEvent)
+{
+	TSharedPtr<FDragDropOperation> Operation = DragDropEvent.GetOperation();
+	if (!Operation.IsValid())
+	{
+		return FReply::Unhandled();
+	}
+	
+	if (Operation->IsOfType<FGeometryCollectionBoneDragDrop>())
+	{
+		TSharedPtr<FGeometryCollectionBoneDragDrop> GeometryCollectionBoneDragDrop = StaticCastSharedPtr<FGeometryCollectionBoneDragDrop>(Operation);
+
+		UGeometryCollectionComponent* SourceComponent = GetComponent();
+
+		FGeometryCollectionEdit GeometryCollectionEdit = SourceComponent->EditRestCollection();
+		if (UGeometryCollection* GeometryCollectionObject = GeometryCollectionEdit.GetRestCollection())
+		{
+			FGeometryCollection* GeometryCollection = GeometryCollectionObject->GetGeometryCollection().Get();
+			if (GeometryCollectionBoneDragDrop->ReparentBones(GeometryCollection, BoneIndex))
+			{
+				ParentComponentItem->RegenerateChildren();
+				ParentComponentItem->RequestTreeRefresh();
+				ParentComponentItem->ExpandAll();
+			}
+		}		
+	}
+	
+	return FReply::Unhandled();
+}
+
 
 UOutlinerSettings::UOutlinerSettings(const FObjectInitializer& ObjInit)
 	: Super(ObjInit)
 	, ItemText(EOutlinerItemNameEnum::BoneIndex)
 {}
 
+
 void SGeometryCollectionOutliner::Construct(const FArguments& InArgs)
 {
 	BoneSelectionChangedDelegate = InArgs._OnBoneSelectionChanged;
@@ -56,60 +144,6 @@ void SGeometryCollectionOutliner::OnGetChildren(FGeometryCollectionTreeItemPtr I
 
 void SGeometryCollectionOutliner::UpdateGeometryCollection()
 {
-	/*if (TSharedPtr<FGeometryCollection, ESPMode::ThreadSafe> CollectionPtr = GeometryCollection.Pin())
-	{
-		RootNodes.Empty();
-		GuidIndexMap.Empty();
-
-		int NumElements = CollectionPtr->NumElements(FGeometryCollection::TransformGroup);
-
-		const TManagedArray<FGuid>& Guids = CollectionPtr->GetAttribute<FGuid>("GUID", "Transform");
-		const TManagedArray<int32>& Parents = CollectionPtr->Parent;
-
-		// Make a copy of the old guids list ... probably a faster way to do this...
-		TArray<FGuid> OldGuids;
-		NodesMap.GenerateKeyArray(OldGuids);
-
-		// Add a sub item to the outliner tree for each of the bones/chunks in this GeometryCollection
-		for (int Index = 0; Index < NumElements; Index++)
-		{
-
-			ItemPtr Item;
-			int32 FoundOldGuidIndex;
-			if (OldGuids.Find(Guids[Index], FoundOldGuidIndex))
-			{
-				// Mark that this item is reused
-				OldGuids.RemoveAt(FoundOldGuidIndex);
-				Item = NodesMap[Guids[Index]];
-			}
-			else
-			{
-				// this is a new bone item, make a new tree item
-				Item = MakeShared<FGeometryCollectionTreeItem>(Guids[Index]);
-				NodesMap.Add(Guids[Index], Item);
-			}
-
-			if (Parents[Index] == FGeometryCollection::Invalid)
-			{
-				RootNodes.Push(Item);
-			}
-
-			GuidIndexMap.Add(Guids[Index], Index);
-		}
-
-		// Remove unused Guids
-		for (auto UnusedGuid : OldGuids)
-		{
-			NodesMap.Remove(UnusedGuid);
-		}
-
-	}
-	else
-	{
-		RootNodes.Empty();
-		GuidIndexMap.Empty();
-		NodesMap.Empty();	
-	}*/
 	TreeView->RequestTreeRefresh();
 	ExpandAll();
 }
@@ -125,9 +159,12 @@ void SGeometryCollectionOutliner::SetComponents(const TArray<UGeometryCollection
 
 	for (UGeometryCollectionComponent* Component : InNewComponents)
 	{
-		RootNodes.Add(MakeShared<FGeometryCollectionTreeItemComponent>(Component));
-		TArray<int32> SelectedBones = Component->GetSelectedBones();
-		SetBoneSelection(Component, SelectedBones, false);
+		if (Component->GetRestCollection() && !Component->GetRestCollection()->IsPendingKill())
+		{
+			RootNodes.Add(MakeShared<FGeometryCollectionTreeItemComponent>(Component, TreeView));
+			TArray<int32> SelectedBones = Component->GetSelectedBones();
+			SetBoneSelection(Component, SelectedBones, false);
+		}
 	}
 
 	TreeView->RequestTreeRefresh();
@@ -138,7 +175,7 @@ void SGeometryCollectionOutliner::ExpandAll()
 {
 	for (TSharedPtr<FGeometryCollectionTreeItemComponent> ItemPtr : RootNodes)
 	{
-		ItemPtr->ExpandAll(TreeView);
+		ItemPtr->ExpandAll();
 	}
 }
 
@@ -330,7 +367,7 @@ FText FGeometryCollectionTreeItemComponent::GetDisplayNameForBone(const FGuid& G
 	return LOCTEXT("BoneNotFound", "Bone Not Found, Invalid Geometry Collection");
 }
 
-void FGeometryCollectionTreeItemComponent::ExpandAll(TSharedPtr<STreeView<FGeometryCollectionTreeItemPtr>> TreeView)
+void FGeometryCollectionTreeItemComponent::ExpandAll()
 {
 	TreeView->SetItemExpansion(AsShared(), true);
 
@@ -366,7 +403,7 @@ void FGeometryCollectionTreeItemComponent::RegenerateChildren()
 			{
 				if (FilterBoneIndex(Index))
 				{
-					TSharedRef<FGeometryCollectionTreeItemBone> NewItem = MakeShared<FGeometryCollectionTreeItemBone>(Guids[Index], Index, *this);
+					TSharedRef<FGeometryCollectionTreeItemBone> NewItem = MakeShared<FGeometryCollectionTreeItemBone>(Guids[Index], Index, this);
 					if (Parents[Index] == RootIndex)
 					{
 						// The actual children directly beneath this node are the ones without a parent.  The rest are children of children
@@ -382,6 +419,14 @@ void FGeometryCollectionTreeItemComponent::RegenerateChildren()
 	}
 }
 
+void FGeometryCollectionTreeItemComponent::RequestTreeRefresh()
+{
+	if (TreeView.IsValid())
+	{
+		TreeView->RequestTreeRefresh();
+	}
+}
+
 void FGeometryCollectionTreeItemComponent::SetHistogramSelection(TArray<int32>& SelectedBones)
 {
 	HistogramSelection = SelectedBones;
@@ -432,26 +477,34 @@ TSharedRef<ITableRow> FGeometryCollectionTreeItemBone::MakeTreeRowWidget(const T
 
 	FSlateColor TextColor(FLinearColor::Red); // default color indicates something wrong
 
-	TSharedPtr<FGeometryCollection, ESPMode::ThreadSafe> GeometryCollectionPtr = ParentComponentItem->GetComponent()->GetRestCollection()->GetGeometryCollection();
-	const TManagedArray<int32>& SimulationType = GeometryCollectionPtr->SimulationType;
-	switch (SimulationType[GetBoneIndex()])
+	const UGeometryCollection* RestCollection = ParentComponentItem->GetComponent()->GetRestCollection();
+	if (RestCollection && !RestCollection->IsPendingKill())
 	{
-		// This is a dubious case. We shouldn't expect to see nodes with this type.
-		case FGeometryCollection::ESimulationTypes::FST_None:
-			TextColor = FColor::Green;
-			break;
+		TSharedPtr<FGeometryCollection, ESPMode::ThreadSafe> GeometryCollectionPtr = RestCollection->GetGeometryCollection();
+		const TManagedArray<int32>& SimulationType = GeometryCollectionPtr->SimulationType;
+		switch (SimulationType[GetBoneIndex()])
+		{
+			case FGeometryCollection::ESimulationTypes::FST_None:
+				TextColor = FColor::Green;
+				break;
 
-		case FGeometryCollection::ESimulationTypes::FST_Rigid:
-			TextColor = FLinearColor::Gray;
-			break;
+			case FGeometryCollection::ESimulationTypes::FST_Rigid:
+				TextColor = FLinearColor::Gray;
+				break;
 
-		case FGeometryCollection::ESimulationTypes::FST_Clustered:
-			TextColor = FColor::Cyan;
-			break;
+			case FGeometryCollection::ESimulationTypes::FST_Clustered:
+				TextColor = FColor::Cyan;
+				break;
 
-		default:
-			ensureMsgf(false, TEXT("Invalid Geometry Collection simulation type encountered."));
-			break;
+			default:
+				ensureMsgf(false, TEXT("Invalid Geometry Collection simulation type encountered."));
+				break;
+		}
+	}
+	else
+	{
+		// Deleted rest collection
+		TextColor = FLinearColor(0.1f, 0.1f, 0.1f);
 	}
 	
 	return SNew(STableRow<FGeometryCollectionTreeItemPtr>, InOwnerTable)
@@ -460,7 +513,11 @@ TSharedRef<ITableRow> FGeometryCollectionTreeItemBone::MakeTreeRowWidget(const T
 			SNew(STextBlock)
 			.Text(ItemText)
 			.ColorAndOpacity(TextColor)
-		];
+		]
+		.OnDragDetected(this, &FGeometryCollectionTreeItem::OnDragDetected)
+		.OnDrop(this, &FGeometryCollectionTreeItem::OnDrop)
+		.OnDragEnter(this, &FGeometryCollectionTreeItem::OnDragEnter)
+		.OnDragLeave(this, &FGeometryCollectionTreeItem::OnDragLeave);
 }
 
 void FGeometryCollectionTreeItemBone::GetChildren(FGeometryCollectionTreeItemList& OutChildren)
@@ -468,4 +525,5 @@ void FGeometryCollectionTreeItemBone::GetChildren(FGeometryCollectionTreeItemLis
 	ParentComponentItem->GetChildrenForBone(*this, OutChildren);
 }
 
+
 #undef LOCTEXT_NAMESPACE // "ChaosEditor"
\ No newline at end of file
diff --git a/Engine/Plugins/Experimental/ChaosEditor/Source/FractureEditor/Private/SGeometryCollectionOutliner.h b/Engine/Plugins/Experimental/ChaosEditor/Source/FractureEditor/Private/SGeometryCollectionOutliner.h
--- a/Engine/Plugins/Experimental/ChaosEditor/Source/FractureEditor/Private/SGeometryCollectionOutliner.h
+++ b/Engine/Plugins/Experimental/ChaosEditor/Source/FractureEditor/Private/SGeometryCollectionOutliner.h
@@ -6,6 +6,7 @@
 #include "Widgets/SCompoundWidget.h"
 #include "Widgets/DeclarativeSyntaxSupport.h"
 #include "Widgets/Views/STreeView.h"
+#include "Input/DragAndDrop.h"
 
 #include "Misc/Attribute.h"
 
@@ -57,13 +58,20 @@ public:
 	virtual UGeometryCollectionComponent* GetComponent() const = 0;
 
 	virtual int32 GetBoneIndex() const { return INDEX_NONE; }
+
+	virtual FReply OnDragDetected(const FGeometry& MyGeometry, const FPointerEvent& MouseEvent) { return FReply::Unhandled(); }
+	virtual FReply OnDrop(const FDragDropEvent& DragDropEvent) { return FReply::Unhandled(); }
+	
+	virtual void OnDragEnter(FDragDropEvent const& InDragDropEvent) {};
+	virtual void OnDragLeave(FDragDropEvent const& InDragDropEvent);
 };
 
 class FGeometryCollectionTreeItemComponent : public FGeometryCollectionTreeItem
 {
 public:
-	FGeometryCollectionTreeItemComponent(UGeometryCollectionComponent* InComponent)
+	FGeometryCollectionTreeItemComponent(UGeometryCollectionComponent* InComponent, TSharedPtr<STreeView<FGeometryCollectionTreeItemPtr>> InTreeView)
 		: Component(InComponent)
+		, TreeView(InTreeView)
 	{
 		RegenerateChildren();
 	}
@@ -79,8 +87,9 @@ public:
 	void GetChildrenForBone(FGeometryCollectionTreeItemBone& BoneItem, FGeometryCollectionTreeItemList& OutChildren) const;
 	FText GetDisplayNameForBone(const FGuid& Guid) const;
 
-	void ExpandAll(TSharedPtr<STreeView<FGeometryCollectionTreeItemPtr>> TreeView);
+	void ExpandAll();
 	void RegenerateChildren();
+	void RequestTreeRefresh();
 
 	void SetHistogramSelection(TArray<int32>& SelectedBones);
 
@@ -90,6 +99,8 @@ private:
 private:
 	TWeakObjectPtr<UGeometryCollectionComponent> Component;
 
+	TSharedPtr<STreeView<FGeometryCollectionTreeItemPtr>> TreeView;
+
 	/** The direct children under this component */
 	TArray<FGeometryCollectionTreeItemPtr> MyChildren;
 
@@ -105,10 +116,10 @@ class FGeometryCollectionTreeItemBone : public FGeometryCollectionTreeItem
 {
 
 public:
-	FGeometryCollectionTreeItemBone(const FGuid NewGuid, const int32 InBoneIndex, const FGeometryCollectionTreeItemComponent& InParentComponentItem)
+	FGeometryCollectionTreeItemBone(const FGuid NewGuid, const int32 InBoneIndex, FGeometryCollectionTreeItemComponent* InParentComponentItem)
 		: Guid(NewGuid)
 		, BoneIndex(InBoneIndex)
-		, ParentComponentItem(&InParentComponentItem)
+		, ParentComponentItem(InParentComponentItem)
 	{}
 
 	/** FGeometryCollectionTreeItem interface */
@@ -118,14 +129,15 @@ public:
 	virtual UGeometryCollectionComponent* GetComponent() const { return ParentComponentItem->GetComponent(); }
 	const FGuid& GetGuid() const { return Guid; }
 
-	// TArray<TSharedPtr<FGeometryCollectionTreeItem>>& GetChildren();
-	// TArray<TSharedPtr<FGeometryCollectionTreeItem>> Children;
-	// TWeakPtr<SGeometryCollectionOutliner> GeomOutliner;
-	
+protected:
+	virtual FReply OnDragDetected(const FGeometry& MyGeometry, const FPointerEvent& MouseEvent) override;
+	virtual FReply OnDrop(const FDragDropEvent& DragDropEvent) override;
+	virtual void OnDragEnter(FDragDropEvent const& InDragDropEvent) override;
+
 private:
 	const FGuid Guid;
 	const int32 BoneIndex;
-	const FGeometryCollectionTreeItemComponent* ParentComponentItem;
+	FGeometryCollectionTreeItemComponent* ParentComponentItem;
 };
 
 class SGeometryCollectionOutliner: public SCompoundWidget
@@ -158,12 +170,10 @@ public:
 	void SetHistogramSelection(UGeometryCollectionComponent* RootComponent, TArray<int32>& SelectedBones);
 
 private:
-	// void GetItemParentsForGuid(int32 Guid);
-	//void GetItemChildrenForGuid(int32 Guid);
 	void OnSelectionChanged(FGeometryCollectionTreeItemPtr Item, ESelectInfo::Type SelectInfo);
 private:
 	TSharedPtr<STreeView<FGeometryCollectionTreeItemPtr>> TreeView;
 	TArray<TSharedPtr<FGeometryCollectionTreeItemComponent>> RootNodes;
 	FOnBoneSelectionChanged BoneSelectionChangedDelegate;
 	bool bPerformingSelection;
-};
\ No newline at end of file
+};
