Fracture Mode -- prevent crash when Geometry Collection component lacks a valid RestCollection.

diff --git a/Engine/Plugins/Experimental/ChaosEditor/Source/FractureEditor/Private/FractureEditorModeToolkit.cpp b/Engine/Plugins/Experimental/ChaosEditor/Source/FractureEditor/Private/FractureEditorModeToolkit.cpp
--- a/Engine/Plugins/Experimental/ChaosEditor/Source/FractureEditor/Private/FractureEditorModeToolkit.cpp
+++ b/Engine/Plugins/Experimental/ChaosEditor/Source/FractureEditor/Private/FractureEditorModeToolkit.cpp
@@ -935,30 +935,33 @@ FText FFractureEditorModeToolkit::GetActiveToolMessage() const
 
 void FFractureEditorModeToolkit::SetOutlinerComponents(const TArray<UGeometryCollectionComponent*>& InNewComponents)
 {
+	TArray<UGeometryCollectionComponent*> ComponentsToEdit;
+	ComponentsToEdit.Reserve(InNewComponents.Num());
 	for (UGeometryCollectionComponent* Component : InNewComponents)
 	{
 		FGeometryCollectionEdit RestCollection = Component->EditRestCollection(GeometryCollection::EEditUpdate::None);
 		UGeometryCollection* FracturedGeometryCollection = RestCollection.GetRestCollection();
 
-		if (FracturedGeometryCollection) // Prevents crash when GC is deleted from content browser and actor is selected.
+		if (FracturedGeometryCollection && !FracturedGeometryCollection->IsPendingKill()) // Prevents crash when GC is deleted from content browser and actor is selected.
 		{
 			TSharedPtr<FGeometryCollection, ESPMode::ThreadSafe> GeometryCollectionPtr = FracturedGeometryCollection->GetGeometryCollection();
 
 			FGeometryCollectionClusteringUtility::UpdateHierarchyLevelOfChildren(GeometryCollectionPtr.Get(), -1);
 			UpdateExplodedVectors(Component);
-		}
 
-		UpdateGeometryComponentAttributes(Component);
+			UpdateGeometryComponentAttributes(Component);
+			ComponentsToEdit.Add(Component);
+		}	
 	}
 
 	if (OutlinerView)
 	{
-		OutlinerView->SetComponents(InNewComponents);
+		OutlinerView->SetComponents(ComponentsToEdit);
 	}
 
 	if (HistogramView)
 	{
-		HistogramView->SetComponents(InNewComponents, GetLevelViewValue());
+		HistogramView->SetComponents(ComponentsToEdit, GetLevelViewValue());
 	}
 
 	if (ActiveTool != nullptr)
@@ -1094,7 +1097,7 @@ void FFractureEditorModeToolkit::UpdateGeometryComponentAttributes(UGeometryColl
 	if (Component)
 	{
 		const UGeometryCollection* RestCollection = Component->GetRestCollection();
-		if (RestCollection)
+		if (RestCollection && !RestCollection->IsPendingKill())
 		{
 			FGeometryCollectionPtr GeometryCollection = RestCollection->GetGeometryCollection();
 
@@ -1337,26 +1340,30 @@ void FFractureEditorModeToolkit::RegenerateHistogram()
 
 void FFractureEditorModeToolkit::OnOutlinerBoneSelectionChanged(UGeometryCollectionComponent* RootComponent, TArray<int32>& SelectedBones)
 {
-	FScopedTransaction Transaction(FractureTransactionContexts::SelectBoneContext, LOCTEXT("SelectGeometryCollectionBoneTransaction", "Select Bone"), RootComponent);
-
-	if(SelectedBones.Num())
-	{
-			
-		FFractureSelectionTools::ToggleSelectedBones(RootComponent, SelectedBones, true);
-		OutlinerView->SetBoneSelection(RootComponent, SelectedBones, true);
-	}
-	else
+	const UGeometryCollection* RestCollection = RootComponent->GetRestCollection();
+	if (RestCollection && !RestCollection->IsPendingKill())
 	{
-		FFractureSelectionTools::ClearSelectedBones(RootComponent);
-	}
+		FScopedTransaction Transaction(FractureTransactionContexts::SelectBoneContext, LOCTEXT("SelectGeometryCollectionBoneTransaction", "Select Bone"), RootComponent);
 
-	if (ActiveTool != nullptr)
-	{
-		ActiveTool->FractureContextChanged();
-	}
+		if (SelectedBones.Num())
+		{
+
+			FFractureSelectionTools::ToggleSelectedBones(RootComponent, SelectedBones, true);
+			OutlinerView->SetBoneSelection(RootComponent, SelectedBones, true);
+		}
+		else
+		{
+			FFractureSelectionTools::ClearSelectedBones(RootComponent);
+		}
 
-	RootComponent->MarkRenderStateDirty();
-	RootComponent->MarkRenderDynamicDataDirty();
+		if (ActiveTool != nullptr)
+		{
+			ActiveTool->FractureContextChanged();
+		}
+
+		RootComponent->MarkRenderStateDirty();
+		RootComponent->MarkRenderDynamicDataDirty();
+	}
 }
 
 void FFractureEditorModeToolkit::OnHistogramBoneSelectionChanged(UGeometryCollectionComponent* RootComponent, TArray<int32>& SelectedBones)
diff --git a/Engine/Plugins/Experimental/ChaosEditor/Source/FractureEditor/Private/FractureToolCutter.cpp b/Engine/Plugins/Experimental/ChaosEditor/Source/FractureEditor/Private/FractureToolCutter.cpp
--- a/Engine/Plugins/Experimental/ChaosEditor/Source/FractureEditor/Private/FractureToolCutter.cpp
+++ b/Engine/Plugins/Experimental/ChaosEditor/Source/FractureEditor/Private/FractureToolCutter.cpp
@@ -45,54 +45,59 @@ TArray<FFractureToolContext> UFractureToolCutterBase::GetFractureToolContexts()
 
 	for (UGeometryCollectionComponent* GeometryCollectionComponent : GeomCompSelection)
 	{
-		// Generate a context for each selected node
-		FFractureToolContext FullSelection(GeometryCollectionComponent);
-		FullSelection.ConvertSelectionToRigidNodes();
-		
-		// Update global transforms and bounds		
-		const TManagedArray<FTransform>& Transform = FullSelection.GetGeometryCollection()->GetAttribute<FTransform>("Transform", FGeometryCollection::TransformGroup);
-		const TManagedArray<int32>& TransformToGeometryIndex = FullSelection.GetGeometryCollection()->GetAttribute<int32>("TransformToGeometryIndex", FGeometryCollection::TransformGroup);
-		const TManagedArray<FBox>& BoundingBoxes = FullSelection.GetGeometryCollection()->GetAttribute<FBox>("BoundingBox", FGeometryCollection::GeometryGroup);
-
-		TArray<FTransform> Transforms;
-		GeometryCollectionAlgo::GlobalMatrices(Transform, FullSelection.GetGeometryCollection()->Parent, Transforms);
-
-		TMap<int32, FBox> BoundsToBone;
-		int32 TransformCount = Transform.Num();
-		for (int32 Index = 0; Index < TransformCount; ++Index)
+		const UGeometryCollection* RestCollection = GeometryCollectionComponent->GetRestCollection();
+		if (RestCollection && !RestCollection->IsPendingKill())
 		{
-			if (TransformToGeometryIndex[Index] > INDEX_NONE)
-			{
-				BoundsToBone.Add(Index, BoundingBoxes[TransformToGeometryIndex[Index]].TransformBy(Transforms[Index]));
-			}
-		}
 
-		if (CutterSettings->bGroupFracture)
-		{
-			FullSelection.SetSeed(CutterSettings->RandomSeed > -1 ? CutterSettings->RandomSeed : FMath::Rand());
+			// Generate a context for each selected node
+			FFractureToolContext FullSelection(GeometryCollectionComponent);
+			FullSelection.ConvertSelectionToRigidNodes();
 
-			FBox Bounds(ForceInit);
-			for (int32 BoneIndex : FullSelection.GetSelection())
+			// Update global transforms and bounds		
+			const TManagedArray<FTransform>& Transform = FullSelection.GetGeometryCollection()->GetAttribute<FTransform>("Transform", FGeometryCollection::TransformGroup);
+			const TManagedArray<int32>& TransformToGeometryIndex = FullSelection.GetGeometryCollection()->GetAttribute<int32>("TransformToGeometryIndex", FGeometryCollection::TransformGroup);
+			const TManagedArray<FBox>& BoundingBoxes = FullSelection.GetGeometryCollection()->GetAttribute<FBox>("BoundingBox", FGeometryCollection::GeometryGroup);
+
+			TArray<FTransform> Transforms;
+			GeometryCollectionAlgo::GlobalMatrices(Transform, FullSelection.GetGeometryCollection()->Parent, Transforms);
+
+			TMap<int32, FBox> BoundsToBone;
+			int32 TransformCount = Transform.Num();
+			for (int32 Index = 0; Index < TransformCount; ++Index)
 			{
-				if (TransformToGeometryIndex[BoneIndex] > INDEX_NONE)
+				if (TransformToGeometryIndex[Index] > INDEX_NONE)
 				{
-					Bounds += BoundsToBone[BoneIndex];
+					BoundsToBone.Add(Index, BoundingBoxes[TransformToGeometryIndex[Index]].TransformBy(Transforms[Index]));
 				}
 			}
-			FullSelection.SetBounds(Bounds);
 
-			Contexts.Add(FullSelection);
-		}
-		else
-		{
-			// Generate a context for each selected node
-			for (int32 Index : FullSelection.GetSelection())
+			if (CutterSettings->bGroupFracture)
 			{
-				Contexts.Emplace(GeometryCollectionComponent);
-				FFractureToolContext& FractureContext = Contexts.Last();
+				FullSelection.SetSeed(CutterSettings->RandomSeed > -1 ? CutterSettings->RandomSeed : FMath::Rand());
 
-				FractureContext.SetSeed(CutterSettings->RandomSeed > -1 ? CutterSettings->RandomSeed + Index : FMath::Rand());
-				FractureContext.SetBounds(BoundsToBone[Index]);
+				FBox Bounds(ForceInit);
+				for (int32 BoneIndex : FullSelection.GetSelection())
+				{
+					if (TransformToGeometryIndex[BoneIndex] > INDEX_NONE)
+					{
+						Bounds += BoundsToBone[BoneIndex];
+					}
+				}
+				FullSelection.SetBounds(Bounds);
+
+				Contexts.Add(FullSelection);
+			}
+			else
+			{
+				// Generate a context for each selected node
+				for (int32 Index : FullSelection.GetSelection())
+				{
+					Contexts.Emplace(GeometryCollectionComponent);
+					FFractureToolContext& FractureContext = Contexts.Last();
+
+					FractureContext.SetSeed(CutterSettings->RandomSeed > -1 ? CutterSettings->RandomSeed + Index : FMath::Rand());
+					FractureContext.SetBounds(BoundsToBone[Index]);
+				}
 			}
 		}
 	}
