manage fracture's default random seed (when user has set seed to -1) separately from getting the fracture context
fixes bug where the "preview cut" diagram would not match the actual cut, because the random seed was re-randomized separately to show the preview and to make the cut.

#rnx

diff --git a/Engine/Plugins/Experimental/ChaosEditor/Source/FractureEditor/Private/FractureToolBrick.cpp b/Engine/Plugins/Experimental/ChaosEditor/Source/FractureEditor/Private/FractureToolBrick.cpp
--- a/Engine/Plugins/Experimental/ChaosEditor/Source/FractureEditor/Private/FractureToolBrick.cpp
+++ b/Engine/Plugins/Experimental/ChaosEditor/Source/FractureEditor/Private/FractureToolBrick.cpp
@@ -245,6 +245,7 @@ void UFractureToolBrick::PostEditChangeChainProperty(struct FPropertyChangedChai
 
 void UFractureToolBrick::FractureContextChanged()
 {
+	UpdateDefaultRandomSeed();
 	UpdateBrickTransforms();
 }
 
diff --git a/Engine/Plugins/Experimental/ChaosEditor/Source/FractureEditor/Private/FractureToolCutter.cpp b/Engine/Plugins/Experimental/ChaosEditor/Source/FractureEditor/Private/FractureToolCutter.cpp
--- a/Engine/Plugins/Experimental/ChaosEditor/Source/FractureEditor/Private/FractureToolCutter.cpp
+++ b/Engine/Plugins/Experimental/ChaosEditor/Source/FractureEditor/Private/FractureToolCutter.cpp
@@ -37,8 +37,6 @@ TArray<FFractureToolContext> UFractureToolCutterBase::GetFractureToolContexts()
 {
 	TArray<FFractureToolContext> Contexts;
 
-	FRandomStream RandomStream(CutterSettings->RandomSeed > -1 ? CutterSettings->RandomSeed : FMath::Rand());
-
 	// A context is gathered for each selected GeometryCollection component, or for each individual bone if Group Fracture is not used.
 	TSet<UGeometryCollectionComponent*> GeomCompSelection;
 	GetSelectedGeometryCollectionComponents(GeomCompSelection);
@@ -73,7 +71,7 @@ TArray<FFractureToolContext> UFractureToolCutterBase::GetFractureToolContexts()
 
 			if (CutterSettings->bGroupFracture)
 			{
-				FullSelection.SetSeed(CutterSettings->RandomSeed > -1 ? CutterSettings->RandomSeed : FMath::Rand());
+				FullSelection.SetSeed(CutterSettings->RandomSeed > -1 ? CutterSettings->RandomSeed : DefaultRandomSeed);
 
 				FBox Bounds(ForceInit);
 				for (int32 BoneIndex : FullSelection.GetSelection())
@@ -95,7 +93,7 @@ TArray<FFractureToolContext> UFractureToolCutterBase::GetFractureToolContexts()
 					Contexts.Emplace(GeometryCollectionComponent);
 					FFractureToolContext& FractureContext = Contexts.Last();
 
-					FractureContext.SetSeed(CutterSettings->RandomSeed > -1 ? CutterSettings->RandomSeed + Index : FMath::Rand());
+					FractureContext.SetSeed(CutterSettings->RandomSeed > -1 ? CutterSettings->RandomSeed + Index : DefaultRandomSeed + Index);
 					FractureContext.SetBounds(BoundsToBone[Index]);
 				}
 			}
@@ -105,6 +103,14 @@ TArray<FFractureToolContext> UFractureToolCutterBase::GetFractureToolContexts()
 	return Contexts;
 }
 
+
+void UFractureToolCutterBase::UpdateDefaultRandomSeed()
+{
+	DefaultRandomSeed = FMath::Rand();
+}
+
+
+
 UFractureToolVoronoiCutterBase::UFractureToolVoronoiCutterBase(const FObjectInitializer& ObjInit)
 	: Super(ObjInit)
 {
@@ -136,11 +142,12 @@ void UFractureToolVoronoiCutterBase::Render(const FSceneView* View, FViewport* V
 
 void UFractureToolVoronoiCutterBase::FractureContextChanged()
 {
+	UpdateDefaultRandomSeed();
 	TArray<FFractureToolContext> FractureContexts = GetFractureToolContexts();
 
 	VoronoiSites.Empty();
 	CellMember.Empty();
-	VoronoiEdges.Empty();;
+	VoronoiEdges.Empty();
 
 	for (FFractureToolContext& FractureContext : FractureContexts)
 	{
diff --git a/Engine/Plugins/Experimental/ChaosEditor/Source/FractureEditor/Private/FractureToolCutter.h b/Engine/Plugins/Experimental/ChaosEditor/Source/FractureEditor/Private/FractureToolCutter.h
--- a/Engine/Plugins/Experimental/ChaosEditor/Source/FractureEditor/Private/FractureToolCutter.h
+++ b/Engine/Plugins/Experimental/ChaosEditor/Source/FractureEditor/Private/FractureToolCutter.h
@@ -102,6 +102,11 @@ public:
 
 	virtual TArray<FFractureToolContext> GetFractureToolContexts() const override;
 
+	virtual void FractureContextChanged() override
+	{
+		UpdateDefaultRandomSeed();
+	}
+
 protected:
 	UPROPERTY(EditAnywhere, Category = Slicing)
 	TObjectPtr<UFractureCutterSettings> CutterSettings;
@@ -109,6 +114,11 @@ protected:
 	UPROPERTY(EditAnywhere, Category = Collision)
 	TObjectPtr<UFractureCollisionSettings> CollisionSettings;
 
+	/// Manage which seed to use when no specific random seed is specified, to control when the seed changes
+	/// (primarily to avoid the seed changing between creating the preview and doing the actual cut)
+	void UpdateDefaultRandomSeed();
+	int32 DefaultRandomSeed = 0;
+
 };
 
 
diff --git a/Engine/Plugins/Experimental/ChaosEditor/Source/FractureEditor/Private/FractureToolPlaneCut.cpp b/Engine/Plugins/Experimental/ChaosEditor/Source/FractureEditor/Private/FractureToolPlaneCut.cpp
--- a/Engine/Plugins/Experimental/ChaosEditor/Source/FractureEditor/Private/FractureToolPlaneCut.cpp
+++ b/Engine/Plugins/Experimental/ChaosEditor/Source/FractureEditor/Private/FractureToolPlaneCut.cpp
@@ -80,6 +80,7 @@ TArray<UObject*> UFractureToolPlaneCut::GetSettingsObjects() const
 
 void UFractureToolPlaneCut::FractureContextChanged()
 {
+	UpdateDefaultRandomSeed();
 	TArray<FFractureToolContext> FractureContexts = GetFractureToolContexts();
 
 	RenderCuttingPlanesTransforms.Empty();
diff --git a/Engine/Plugins/Experimental/ChaosEditor/Source/FractureEditor/Private/FractureToolResample.cpp b/Engine/Plugins/Experimental/ChaosEditor/Source/FractureEditor/Private/FractureToolResample.cpp
--- a/Engine/Plugins/Experimental/ChaosEditor/Source/FractureEditor/Private/FractureToolResample.cpp
+++ b/Engine/Plugins/Experimental/ChaosEditor/Source/FractureEditor/Private/FractureToolResample.cpp
@@ -57,6 +57,7 @@ TArray<UObject*> UFractureToolResample::GetSettingsObjects() const
 
 void UFractureToolResample::FractureContextChanged()
 {
+	UpdateDefaultRandomSeed();
 	TArray<FFractureToolContext> FractureContexts = GetFractureToolContexts();
 
 	GeneratedPoints.Reset();
@@ -78,8 +79,16 @@ void UFractureToolResample::FractureContextChanged()
 			FTransform CombinedTransform = InnerTransform * OuterTransform;
 			int32 GeometryIdx = Collection.TransformToGeometryIndex[TransformIdx];
 			int32 VertStart = Collection.VertexStart[GeometryIdx];
-			int32 VertCount = Collection.VertexCount[GeometryIdx];
-			for (int32 VIdx = VertStart; VIdx < VertStart + VertCount; VIdx++)
+			int32 VertEnd = VertStart + Collection.VertexCount[GeometryIdx];
+			int32 FaceStart = Collection.FaceStart[GeometryIdx];
+			int32 FaceEnd = FaceStart + Collection.FaceCount[GeometryIdx];
+			// only show off-vertex samples; skip over the samples that are on faces
+			for (int32 FIdx = FaceStart; FIdx < FaceEnd; FIdx++)
+			{
+				FIntVector Face = Collection.Indices[FIdx];
+				VertStart = FMath::Max(VertStart, Face.GetMax() + 1);
+			}
+			for (int32 VIdx = VertStart; VIdx < VertEnd; VIdx++)
 			{
 				GeneratedPoints.Add(CombinedTransform.TransformPosition(Collection.Vertex[VIdx]));
 			}
diff --git a/Engine/Plugins/Experimental/ChaosEditor/Source/FractureEditor/Private/FractureToolSlice.cpp b/Engine/Plugins/Experimental/ChaosEditor/Source/FractureEditor/Private/FractureToolSlice.cpp
--- a/Engine/Plugins/Experimental/ChaosEditor/Source/FractureEditor/Private/FractureToolSlice.cpp
+++ b/Engine/Plugins/Experimental/ChaosEditor/Source/FractureEditor/Private/FractureToolSlice.cpp
@@ -126,6 +126,7 @@ void UFractureToolSlice::Render(const FSceneView* View, FViewport* Viewport, FPr
 
 void UFractureToolSlice::FractureContextChanged()
 {
+	UpdateDefaultRandomSeed();
 	TArray<FFractureToolContext> FractureContexts = GetFractureToolContexts();
 
 	RenderCuttingPlanesTransforms.Empty();
