typo fix

diff --git a/Engine/Plugins/Experimental/ChaosEditor/Source/FractureEditor/Private/FractureEditorModeToolkit.cpp b/Engine/Plugins/Experimental/ChaosEditor/Source/FractureEditor/Private/FractureEditorModeToolkit.cpp
--- a/Engine/Plugins/Experimental/ChaosEditor/Source/FractureEditor/Private/FractureEditorModeToolkit.cpp
+++ b/Engine/Plugins/Experimental/ChaosEditor/Source/FractureEditor/Private/FractureEditorModeToolkit.cpp
@@ -845,7 +845,7 @@ TSharedRef<SWidget> FFractureEditorModeToolkit::GetLevelViewMenuContent(TSharedR
 	{
 		MenuBuilder.AddMenuEntry(
 			FText::Format(LOCTEXT("LevelMenuN", "Level {0}"), FText::AsNumber(i)),
-			FText::Format(LOCTEXT("LevelMenuNTooltip", "View Level {0} in this Geometry Collecdtion"), FText::AsNumber(i)),
+			FText::Format(LOCTEXT("LevelMenuNTooltip", "View Level {0} in this Geometry Collection"), FText::AsNumber(i)),
 			FSlateIcon(),
 			FUIAction(
 				FExecuteAction::CreateLambda([=] { PropertyHandle->SetValue(i); } ),
