Fracture Mode
-- New Autocluster mode: Voronoi
-- Proximity generation on demand.

diff --git a/Engine/Plugins/Experimental/ChaosEditor/Source/FractureEditor/Private/AutoClusterFracture.h b/Engine/Plugins/Experimental/ChaosEditor/Source/FractureEditor/Private/AutoClusterFracture.h
--- a/Engine/Plugins/Experimental/ChaosEditor/Source/FractureEditor/Private/AutoClusterFracture.h
+++ b/Engine/Plugins/Experimental/ChaosEditor/Source/FractureEditor/Private/AutoClusterFracture.h
@@ -20,6 +20,8 @@ enum class EFractureAutoClusterMode : uint8
 
 	/** Distance */
 	Distance UMETA(DisplayName = "Distance"),
+
+	Voronoi UMETA(DisplayName = "Voronoi"),
 };
 
 /** Performs clustering of the currently selected geometry collection bones */
diff --git a/Engine/Plugins/Experimental/ChaosEditor/Source/FractureEditor/Private/FractureToolAutoCluster.cpp b/Engine/Plugins/Experimental/ChaosEditor/Source/FractureEditor/Private/FractureToolAutoCluster.cpp
--- a/Engine/Plugins/Experimental/ChaosEditor/Source/FractureEditor/Private/FractureToolAutoCluster.cpp
+++ b/Engine/Plugins/Experimental/ChaosEditor/Source/FractureEditor/Private/FractureToolAutoCluster.cpp
@@ -9,6 +9,9 @@
 #include "AutoClusterFracture.h"
 
 #include "GeometryCollection/GeometryCollectionComponent.h"
+#include "GeometryCollection/GeometryCollectionClusteringUtility.h"
+#include "GeometryCollection/GeometryCollectionAlgo.h"
+#include "Math/NumericLimits.h"
 
 #define LOCTEXT_NAMESPACE "FractureAutoCluster"
 
@@ -64,12 +67,39 @@ void UFractureToolAutoCluster::Execute(TWeakPtr<FFractureEditorModeToolkit> InTo
 		FFractureEditorModeToolkit* Toolkit = InToolkit.Pin().Get();
 
 		TArray<FFractureToolContext> Contexts = GetFractureToolContexts();
-
+		
 		for (FFractureToolContext& Context : Contexts)
 		{
 			Context.ConvertSelectionToClusterNodes();
 
-			UAutoClusterFractureCommand::ClusterChildBonesOfASingleMesh(Context.GetGeometryCollectionComponent(), AutoClusterSettings->AutoClusterMode, AutoClusterSettings->SiteCount);
+			if (AutoClusterSettings->AutoClusterMode < EFractureAutoClusterMode::Voronoi)
+			{
+				UAutoClusterFractureCommand::ClusterChildBonesOfASingleMesh(Context.GetGeometryCollectionComponent(), AutoClusterSettings->AutoClusterMode, AutoClusterSettings->SiteCount);
+			}
+			else
+			{
+				for (const int32 ClusterIndex : Context.GetSelection())
+				{
+					FVoronoiPartitioner VoronoiPartition(Context.GetGeometryCollection().Get(), ClusterIndex);
+					VoronoiPartition.KMeansPartition(AutoClusterSettings->SiteCount);
+
+					if (AutoClusterSettings->bEnforceConnectivity)
+					{
+						GenerateProximityIfNecessary(Context.GetGeometryCollection().Get());
+						VoronoiPartition.SplitDisconnectedPartitions(Context.GetGeometryCollection().Get());
+					}
+					
+					int32 PartitionCount = VoronoiPartition.GetPartitionCount();
+					for (int32 NewClusterIndex = 0; NewClusterIndex < VoronoiPartition.GetPartitionCount(); ++NewClusterIndex)
+					{
+						TArray<int32> NewCluster = VoronoiPartition.GetPartition(NewClusterIndex);
+						if (NewCluster.Num() > 0)
+						{
+							FGeometryCollectionClusteringUtility::ClusterBonesUnderNewNode(Context.GetGeometryCollection().Get(), NewCluster[0], NewCluster, true);
+						}
+					}
+				}
+			}
 
 			Refresh(Context, Toolkit);
 		}
@@ -77,6 +107,267 @@ void UFractureToolAutoCluster::Execute(TWeakPtr<FFractureEditorModeToolkit> InTo
 		SetOutlinerComponents(Contexts, Toolkit);
 	}
 }
+
+
+FVoronoiPartitioner::FVoronoiPartitioner(const FGeometryCollection* GeometryCollection, int32 ClusterIndex)
+{
+	// Collect children of cluster
+	const TManagedArray<TSet<int32>>& Children = GeometryCollection->Children;
+	TransformIndices = Children[ClusterIndex].Array();
+
+	GenerateCentroids(GeometryCollection);
+}
+
+void FVoronoiPartitioner::KMeansPartition(int32 InPartitionCount)
+{
+	PartitionCount = InPartitionCount;
+
+	InitializePartitions();
+	for (int32 Iteration = 0; Iteration < MaxKMeansIterations; Iteration++)
+	{
+		// Refinement is complete when no nodes change partition.
+		if (!Refine())
+		{
+			break;
+		}
+	}
+}
+
+void FVoronoiPartitioner::SplitDisconnectedPartitions(FGeometryCollection* GeometryCollection)
+{
+	Visited.Init(false, TransformIndices.Num());
+	GenerateConnectivity(GeometryCollection);
+	
+	for (int32 PartitionIndex = 0; PartitionIndex < PartitionCount; ++PartitionIndex)
+	{
+		int32 FirstElementInPartitionIndex = Partitions.Find(PartitionIndex);
+		if (FirstElementInPartitionIndex > INDEX_NONE)
+		{
+			MarkVisited(FirstElementInPartitionIndex, PartitionIndex);
+
+			// Place unvisited partition members in a new partition.
+			bool bFoundUnattached = false;
+			for (int32 Index = 0; Index < TransformIndices.Num(); ++Index)
+			{
+				if (Partitions[Index] == PartitionIndex)
+				{
+					if (!Visited[Index])
+					{
+						if (!bFoundUnattached)
+						{
+							bFoundUnattached = true;
+							PartitionCount++;
+						}
+						Partitions[Index] = PartitionCount - 1;
+					}
+				}
+			}
+		}
+	}
+}
+
+void FVoronoiPartitioner::MarkVisited(int32 Index, int32 PartitionIndex)
+{
+	Visited[Index] = true;
+	for (int32 Adjacent : Connectivity[Index])
+	{
+		if (!Visited[Adjacent] && Partitions[Adjacent] == PartitionIndex)
+		{
+			MarkVisited(Adjacent, PartitionIndex);
+		}
+	}
+}
+
+TArray<int32> FVoronoiPartitioner::GetPartition(int32 PartitionIndex) const
+{
+	TArray<int32> PartitionContents;
+	PartitionContents.Reserve(TransformIndices.Num());
+	for (int32 Index = 0; Index < TransformIndices.Num(); ++Index)
+	{
+		if (Partitions[Index] == PartitionIndex)
+		{
+			PartitionContents.Add(TransformIndices[Index]);
+		}
+	}
+
+	return PartitionContents;
+}
+
+void FVoronoiPartitioner::GenerateConnectivity(const FGeometryCollection* GeometryCollection)
+{
+	
+	Connectivity.SetNum(TransformIndices.Num());
+
+	const TManagedArray<int32>& Levels = GeometryCollection->GetAttribute<int32>("Level", FGeometryCollection::TransformGroup);
+
+	for (int32 Index = 0; Index < TransformIndices.Num(); ++Index)
+	{
+		CollectConnections(GeometryCollection, TransformIndices[Index], Levels[TransformIndices[Index]], Connectivity[Index]);
+	}
+}
+
+void FVoronoiPartitioner::CollectConnections(const FGeometryCollection* GeometryCollection, int32 Index, int32 OperatingLevel, TSet<int32>& OutConnections) const
+{
+	const TManagedArray<TSet<int32>>& Children = GeometryCollection->Children;
+	if (GeometryCollection->IsRigid(Index))
+	{
+		const TManagedArray<TSet<int32>>& Proximity = GeometryCollection->GetAttribute<TSet<int32>>("Proximity", FGeometryCollection::GeometryGroup);
+		const TManagedArray<int32>& GeometryToTransformIndex = GeometryCollection->TransformIndex;
+		const TManagedArray<int32>& TransformToGeometryIndex = GeometryCollection->TransformToGeometryIndex;
+
+		for (int32 Neighbor : Proximity[TransformToGeometryIndex[Index]])
+		{
+			int32 NeighborTransformIndex = FGeometryCollectionClusteringUtility::GetParentOfBoneAtSpecifiedLevel(GeometryCollection, GeometryToTransformIndex[Neighbor], OperatingLevel);
+			int32 ContextIndex = TransformIndices.Find(NeighborTransformIndex);
+			if (ContextIndex > INDEX_NONE)
+			{
+				OutConnections.Add(ContextIndex);
+			}
+		}
+	}
+	else
+	{
+		for (int32 ChildIndex : Children[Index])
+		{
+			CollectConnections(GeometryCollection, ChildIndex, OperatingLevel, OutConnections);
+		}
+	}
+}
+
+void FVoronoiPartitioner::GenerateCentroids(const FGeometryCollection* GeometryCollection)
+{
+	Centroids.SetNum(TransformIndices.Num());
+	for (int32 Index = 0; Index < Centroids.Num(); ++Index)
+	{
+		Centroids[Index] = GenerateCentroid(GeometryCollection, TransformIndices[Index]);
+	}
+}
+
+FVector FVoronoiPartitioner::GenerateCentroid(const FGeometryCollection* GeometryCollection, int32 TransformIndex) const
+{
+	FBox Bounds = GenerateBounds(GeometryCollection, TransformIndex);
+	return Bounds.GetCenter();
+}
+
+FBox FVoronoiPartitioner::GenerateBounds(const FGeometryCollection* GeometryCollection, int32 TransformIndex) const
+{
+	check(GeometryCollection->IsRigid(TransformIndex) || GeometryCollection->IsClustered(TransformIndex));
+
+	// Return the bounds of all the vertices contained by this branch.
+
+	if (GeometryCollection->IsRigid(TransformIndex))
+	{
+		const TManagedArray<FTransform>& Transforms = GeometryCollection->Transform;
+		const TManagedArray<int32>& Parents = GeometryCollection->Parent;
+		FTransform GlobalTransform = GeometryCollectionAlgo::GlobalMatrix(Transforms, Parents, TransformIndex);
+		
+		int32 GeometryIndex = GeometryCollection->TransformToGeometryIndex[TransformIndex];
+		int32 VertexStart = GeometryCollection->VertexStart[GeometryIndex];
+		int32 VertexCount = GeometryCollection->VertexCount[GeometryIndex];
+		const TManagedArray<FVector>& GCVertices = GeometryCollection->Vertex;
+
+		TArray<FVector> Vertices;
+		Vertices.SetNum(VertexCount);
+		for (int32 VertexOffset = 0; VertexOffset < VertexCount; ++VertexOffset)
+		{
+			Vertices[VertexOffset] = GlobalTransform.TransformPosition(GCVertices[VertexStart + VertexOffset]);
+		}
+
+		return FBox(Vertices);
+	}
+	else // recurse the cluster
+	{
+		const TArray<int32> Children = GeometryCollection->Children[TransformIndex].Array();
+		check(Children.Num() > 0);
+		FBox Bounds = GenerateBounds(GeometryCollection, Children[0]);
+		for (int32 ChildIndex = 1; ChildIndex < Children.Num(); ++ChildIndex)
+		{
+			Bounds += GenerateBounds(GeometryCollection, Children[ChildIndex]);
+		}
+
+		return Bounds;
+	}
+}
+
+void FVoronoiPartitioner::InitializePartitions()
+{
+	check(PartitionCount > 0);
+
+	PartitionCount = FMath::Min(PartitionCount, TransformIndices.Num());
+
+	// Set initial partition centers as regular selects from the vertex set
+	PartitionCenters.SetNum(PartitionCount);
+	int32 TransformStride = FMath::CeilToInt(float(TransformIndices.Num()) / float(PartitionCount));
+	for (int32 PartitionIndex = 0; PartitionIndex < PartitionCount; ++PartitionIndex)
+	{
+		PartitionCenters[PartitionIndex] = Centroids[TransformStride * PartitionIndex];
+	}
+
+	// At beginning, all nodes belong to first partition.
+	Partitions.Init(0, TransformIndices.Num());
+	PartitionSize.Init(0, PartitionCount);
+	PartitionSize[0] = TransformIndices.Num();
+}
+
+bool FVoronoiPartitioner::Refine()
+{
+	// Assign each transform to its closest partition center.
+	bool bUnchanged = true;
+	for (int32 Index = 0; Index < Centroids.Num(); ++Index)
+	{
+		int32 ClosestPartition = FindClosestPartitionCenter(Centroids[Index]);
+		if (ClosestPartition != Partitions[Index])
+		{
+			bUnchanged = false;
+			PartitionSize[Partitions[Index]]--;
+			PartitionSize[ClosestPartition]++;
+			Partitions[Index] = ClosestPartition;
+		}
+	}
+
+	if (bUnchanged)
+	{
+		return false;
+	}
+
+	// Recalculate partition centers
+	for (int32 PartitionIndex = 0; PartitionIndex < PartitionCount; ++PartitionIndex)
+	{
+		PartitionCenters[PartitionIndex] = FVector(0.0f, 0.0f, 0.0f);
+	}
+
+	for (int32 Index = 0; Index < Partitions.Num(); ++Index)
+	{
+		PartitionCenters[Partitions[Index]] += Centroids[Index];
+	}
+
+	for (int32 PartitionIndex = 0; PartitionIndex < PartitionCount; ++PartitionIndex)
+	{
+		check(PartitionSize[PartitionIndex] > 0);
+		PartitionCenters[PartitionIndex] /= float(PartitionSize[PartitionIndex]);
+	}
+
+	return true;
+}
+
+int32 FVoronoiPartitioner::FindClosestPartitionCenter(const FVector& Location) const
+{
+	int32 ClosestPartition = INDEX_NONE;
+	float SmallestDistSquared = TNumericLimits<float>::Max();
+
+	for (int32 PartitionIndex = 0; PartitionIndex < PartitionCount; ++PartitionIndex)
+	{
+		float DistSquared = FVector::DistSquared(Location, PartitionCenters[PartitionIndex]);
+		if (DistSquared < SmallestDistSquared)
+		{
+			SmallestDistSquared = DistSquared;
+			ClosestPartition = PartitionIndex;
+		}
+	}
+
+	check(ClosestPartition > INDEX_NONE);
+	return ClosestPartition;
+}
 	
 #undef LOCTEXT_NAMESPACE
 
diff --git a/Engine/Plugins/Experimental/ChaosEditor/Source/FractureEditor/Private/FractureToolAutoCluster.h b/Engine/Plugins/Experimental/ChaosEditor/Source/FractureEditor/Private/FractureToolAutoCluster.h
--- a/Engine/Plugins/Experimental/ChaosEditor/Source/FractureEditor/Private/FractureToolAutoCluster.h
+++ b/Engine/Plugins/Experimental/ChaosEditor/Source/FractureEditor/Private/FractureToolAutoCluster.h
@@ -23,8 +23,11 @@ public:
 	UPROPERTY(EditAnywhere, Category = AutoCluster, meta = (DisplayName = "Mode"))
 	EFractureAutoClusterMode AutoClusterMode;
 
-	UPROPERTY(EditAnywhere, Category = AutoCluster, meta = (DisplayName = "Cluster Sites"))// , UIMin = "1", UIMax = "5000", ClampMin = "1"))
-	uint32 SiteCount;
+	UPROPERTY(EditAnywhere, Category = AutoCluster, meta = (DisplayName = "Cluster Sites", UIMin = "1", UIMax = "5000", ClampMin = "1"))
+	uint32 SiteCount=1;
+
+	UPROPERTY(EditAnywhere, Category = AutoCluster, meta = (DisplayName = "Enforce Cluster Connectivity"))
+	bool bEnforceConnectivity=true;
 };
 
 
@@ -48,4 +51,47 @@ public:
 
 	UPROPERTY(EditAnywhere, Category = AutoCluster)
 	TObjectPtr<UFractureAutoClusterSettings> AutoClusterSettings;
+};
+
+
+class FVoronoiPartitioner
+{
+public:
+	FVoronoiPartitioner(const FGeometryCollection* GeometryCollection, int32 ClusterIndex);
+
+	/** Cluster bodies into k partitions using K-Means. Connectivity is ignored: only spatial proximity is considered. */
+	void KMeansPartition(int32 InPartitionCount);
+
+	/** Split any partition islands into their own partition. This will possbily increase number of partitions to exceed desired count. */
+	void SplitDisconnectedPartitions(FGeometryCollection* GeometryCollection);
+
+	int32 GetPartitionCount() const { return PartitionCount; }
+
+	/** return the GeometryCollection TranformIndices within the partition. */
+	TArray<int32> GetPartition(int32 PartitionIndex) const;
+
+private:
+	void GenerateConnectivity(const FGeometryCollection* GeometryCollection);
+	void CollectConnections(const FGeometryCollection* GeometryCollection, int32 Index, int32 OperatingLevel, TSet<int32>& OutConnections) const;
+	void GenerateCentroids(const FGeometryCollection* GeometryCollection);
+	FVector GenerateCentroid(const FGeometryCollection* GeometryCollection, int32 TransformIndex) const;
+	FBox GenerateBounds(const FGeometryCollection* GeometryCollection, int32 TransformIndex) const;
+	void InitializePartitions();
+	bool Refine();
+	int32 FindClosestPartitionCenter(const FVector& Location) const;
+	void MarkVisited(int32 Index, int32 PartitionIndex);
+
+private:
+	TArray<int32> TransformIndices;
+	TArray<FVector> Centroids;
+	TArray<int32> Partitions;
+	int32 PartitionCount;
+	TArray<int32> PartitionSize;
+	TArray<FVector> PartitionCenters;
+	TArray<TSet<int32>> Connectivity;
+	TArray<bool> Visited;
+
+
+	// Not generally necessary but this is a safety measure to prevent oscillating solves that never converge.
+	const int32 MaxKMeansIterations = 500;
 };
\ No newline at end of file
diff --git a/Engine/Plugins/Experimental/ChaosEditor/Source/FractureEditor/Private/FractureToolBrick.cpp b/Engine/Plugins/Experimental/ChaosEditor/Source/FractureEditor/Private/FractureToolBrick.cpp
--- a/Engine/Plugins/Experimental/ChaosEditor/Source/FractureEditor/Private/FractureToolBrick.cpp
+++ b/Engine/Plugins/Experimental/ChaosEditor/Source/FractureEditor/Private/FractureToolBrick.cpp
@@ -327,6 +327,9 @@ int32 UFractureToolBrick::ExecuteFracture(const FFractureToolContext& FractureCo
 			VoronoiPlanarCells.InternalSurfaceMaterials.NoiseSettings = NoiseSettings;
 		}
 
+		// Proximity is invalidated.
+		ClearProximity(FractureContext.GetGeometryCollection().Get());
+
 		return CutMultipleWithPlanarCells(VoronoiPlanarCells, *FractureContext.GetGeometryCollection(), FractureContext.GetSelection(), 0, CollisionSettings->PointSpacing, FractureContext.GetTransform());
 	}
 
diff --git a/Engine/Plugins/Experimental/ChaosEditor/Source/FractureEditor/Private/FractureToolCutter.cpp b/Engine/Plugins/Experimental/ChaosEditor/Source/FractureEditor/Private/FractureToolCutter.cpp
--- a/Engine/Plugins/Experimental/ChaosEditor/Source/FractureEditor/Private/FractureToolCutter.cpp
+++ b/Engine/Plugins/Experimental/ChaosEditor/Source/FractureEditor/Private/FractureToolCutter.cpp
@@ -185,6 +185,9 @@ int32 UFractureToolVoronoiCutterBase::ExecuteFracture(const FFractureToolContext
 			VoronoiPlanarCells.InternalSurfaceMaterials.NoiseSettings = NoiseSettings;
 		}
 
+		// Proximity is invalidated.
+		ClearProximity(FractureContext.GetGeometryCollection().Get());
+
 		return CutMultipleWithPlanarCells(VoronoiPlanarCells, *(FractureContext.GetGeometryCollection()), FractureContext.GetSelection(), CutterSettings->Grout, CollisionSettings->PointSpacing, FractureContext.GetTransform());
 	}
 
diff --git a/Engine/Plugins/Experimental/ChaosEditor/Source/FractureEditor/Private/FractureToolEditing.cpp b/Engine/Plugins/Experimental/ChaosEditor/Source/FractureEditor/Private/FractureToolEditing.cpp
--- a/Engine/Plugins/Experimental/ChaosEditor/Source/FractureEditor/Private/FractureToolEditing.cpp
+++ b/Engine/Plugins/Experimental/ChaosEditor/Source/FractureEditor/Private/FractureToolEditing.cpp
@@ -86,6 +86,9 @@ void UFractureToolDeleteBranch::Execute(TWeakPtr<FFractureEditorModeToolkit> InT
 
 			Context.GetGeometryCollectionComponent()->InitializeEmbeddedGeometry();
 
+			// Proximity is invalidated.
+			ClearProximity(Context.GetGeometryCollection().Get());
+
 			Refresh(Context, Toolkit, true);
 		}
 
diff --git a/Engine/Plugins/Experimental/ChaosEditor/Source/FractureEditor/Private/FractureToolMeshCut.cpp b/Engine/Plugins/Experimental/ChaosEditor/Source/FractureEditor/Private/FractureToolMeshCut.cpp
--- a/Engine/Plugins/Experimental/ChaosEditor/Source/FractureEditor/Private/FractureToolMeshCut.cpp
+++ b/Engine/Plugins/Experimental/ChaosEditor/Source/FractureEditor/Private/FractureToolMeshCut.cpp
@@ -88,6 +88,10 @@ int32 UFractureToolMeshCut::ExecuteFracture(const FFractureToolContext& Fracture
 		FMeshDescription* MeshDescription = LocalCutSettings->CuttingActor->GetStaticMeshComponent()->GetStaticMesh()->GetMeshDescription(0);
 		FTransform Transform(LocalCutSettings->CuttingActor->GetTransform());
 		FInternalSurfaceMaterials InternalSurfaceMaterials;
+
+		// Proximity is invalidated.
+		ClearProximity(FractureContext.GetGeometryCollection().Get());
+
 		// (Note: noise and grout not currently supported)
 		return CutWithMesh(MeshDescription, Transform, InternalSurfaceMaterials, *FractureContext.GetGeometryCollection(), FractureContext.GetSelection(), CollisionSettings->PointSpacing, FractureContext.GetTransform());
 	}
diff --git a/Engine/Plugins/Experimental/ChaosEditor/Source/FractureEditor/Private/FractureToolPlaneCut.cpp b/Engine/Plugins/Experimental/ChaosEditor/Source/FractureEditor/Private/FractureToolPlaneCut.cpp
--- a/Engine/Plugins/Experimental/ChaosEditor/Source/FractureEditor/Private/FractureToolPlaneCut.cpp
+++ b/Engine/Plugins/Experimental/ChaosEditor/Source/FractureEditor/Private/FractureToolPlaneCut.cpp
@@ -132,6 +132,9 @@ int32 UFractureToolPlaneCut::ExecuteFracture(const FFractureToolContext& Fractur
 			InternalSurfaceMaterials.NoiseSettings = NoiseSettings;
 		}
 
+		// Proximity is invalidated.
+		ClearProximity(FractureContext.GetGeometryCollection().Get());
+
 		return CutMultipleWithMultiplePlanes(CuttingPlanes, InternalSurfaceMaterials, *FractureContext.GetGeometryCollection(), FractureContext.GetSelection(), CutterSettings->Grout, CollisionSettings->PointSpacing, FractureContext.GetTransform());
 	}
 
diff --git a/Engine/Plugins/Experimental/ChaosEditor/Source/FractureEditor/Private/FractureToolSlice.cpp b/Engine/Plugins/Experimental/ChaosEditor/Source/FractureEditor/Private/FractureToolSlice.cpp
--- a/Engine/Plugins/Experimental/ChaosEditor/Source/FractureEditor/Private/FractureToolSlice.cpp
+++ b/Engine/Plugins/Experimental/ChaosEditor/Source/FractureEditor/Private/FractureToolSlice.cpp
@@ -171,6 +171,9 @@ int32 UFractureToolSlice::ExecuteFracture(const FFractureToolContext& FractureCo
 			InternalSurfaceMaterials.NoiseSettings = NoiseSettings;
 		}
 
+		// Proximity is invalidated.
+		ClearProximity(FractureContext.GetGeometryCollection().Get());
+
 		return CutMultipleWithMultiplePlanes(CuttingPlanes, InternalSurfaceMaterials, *FractureContext.GetGeometryCollection(), FractureContext.GetSelection(), CutterSettings->Grout, CollisionSettings->PointSpacing);
 	}
 
diff --git a/Engine/Plugins/Experimental/ChaosEditor/Source/FractureEditor/Public/FractureTool.cpp b/Engine/Plugins/Experimental/ChaosEditor/Source/FractureEditor/Public/FractureTool.cpp
--- a/Engine/Plugins/Experimental/ChaosEditor/Source/FractureEditor/Public/FractureTool.cpp
+++ b/Engine/Plugins/Experimental/ChaosEditor/Source/FractureEditor/Public/FractureTool.cpp
@@ -12,6 +12,7 @@
 #include "GeometryCollection/GeometryCollection.h"
 #include "GeometryCollection/GeometryCollectionComponent.h"
 #include "GeometryCollection/GeometryCollectionClusteringUtility.h"
+#include "GeometryCollection/GeometryCollectionProximityUtility.h"
 #include "GeometryCollection/GeometryCollectionAlgo.h"
 
 
@@ -170,6 +171,22 @@ void UFractureActionTool::SetOutlinerComponents(TArray<FFractureToolContext>& In
 	Toolkit->SetOutlinerComponents(Components);
 }
 
+void UFractureActionTool::ClearProximity(FGeometryCollection* GeometryCollection)
+{
+	if (GeometryCollection->HasAttribute("Proximity", FGeometryCollection::GeometryGroup))
+	{
+		GeometryCollection->RemoveAttribute("Proximity", FGeometryCollection::GeometryGroup);
+	}
+}
+
+void UFractureActionTool::GenerateProximityIfNecessary(FGeometryCollection* GeometryCollection)
+{
+	if (!GeometryCollection->HasAttribute("Proximity", FGeometryCollection::GeometryGroup))
+	{
+		FGeometryCollectionProximityUtility::UpdateProximity(GeometryCollection);
+	}
+}
+
 TArray<FFractureToolContext> UFractureActionTool::GetFractureToolContexts() const
 {
 	TArray<FFractureToolContext> Contexts;
diff --git a/Engine/Plugins/Experimental/ChaosEditor/Source/FractureEditor/Public/FractureTool.h b/Engine/Plugins/Experimental/ChaosEditor/Source/FractureEditor/Public/FractureTool.h
--- a/Engine/Plugins/Experimental/ChaosEditor/Source/FractureEditor/Public/FractureTool.h
+++ b/Engine/Plugins/Experimental/ChaosEditor/Source/FractureEditor/Public/FractureTool.h
@@ -70,6 +70,8 @@ protected:
 	static void GetSelectedGeometryCollectionComponents(TSet<UGeometryCollectionComponent*>& GeomCompSelection);
 	static void Refresh(FFractureToolContext& Context, FFractureEditorModeToolkit* Toolkit, bool bClearSelection=false);
 	static void SetOutlinerComponents(TArray<FFractureToolContext>& InContexts, FFractureEditorModeToolkit* Toolkit);
+	static void ClearProximity(FGeometryCollection* GeometryCollection);
+	static void GenerateProximityIfNecessary(FGeometryCollection* GeometryCollection);
 
 
 protected:
