Fracture Tool
-- Cleanup of logging and little bug fix on Auto Cluster Voronoi

diff --git a/Engine/Plugins/Experimental/ChaosEditor/Source/FractureEditor/Private/FractureToolAutoCluster.cpp b/Engine/Plugins/Experimental/ChaosEditor/Source/FractureEditor/Private/FractureToolAutoCluster.cpp
--- a/Engine/Plugins/Experimental/ChaosEditor/Source/FractureEditor/Private/FractureToolAutoCluster.cpp
+++ b/Engine/Plugins/Experimental/ChaosEditor/Source/FractureEditor/Private/FractureToolAutoCluster.cpp
@@ -130,7 +130,6 @@ void FVoronoiPartitioner::KMeansPartition(int32 InPartitionCount)
 		{
 			break;
 		}
-		UE_LOG(LogTemp, Warning, TEXT("Iter %d"), Iteration);
 	}
 }
 
@@ -298,7 +297,7 @@ void FVoronoiPartitioner::InitializePartitions()
 
 	// Set initial partition centers as selects from the vertex set
 	PartitionCenters.SetNum(PartitionCount);
-	int32 TransformStride = FMath::Min(1,FMath::FloorToInt(float(TransformIndices.Num()) / float(PartitionCount)));
+	int32 TransformStride = FMath::Max(1,FMath::FloorToInt(float(TransformIndices.Num()) / float(PartitionCount)));
 	for (int32 PartitionIndex = 0; PartitionIndex < PartitionCount; ++PartitionIndex)
 	{
 		PartitionCenters[PartitionIndex] = Centroids[TransformStride * PartitionIndex];
