change default number of plane cuts in fracture tool from 3 to 1

#rnx

diff --git a/Engine/Plugins/Experimental/ChaosEditor/Source/FractureEditor/Private/FractureToolPlaneCut.h b/Engine/Plugins/Experimental/ChaosEditor/Source/FractureEditor/Private/FractureToolPlaneCut.h
--- a/Engine/Plugins/Experimental/ChaosEditor/Source/FractureEditor/Private/FractureToolPlaneCut.h
+++ b/Engine/Plugins/Experimental/ChaosEditor/Source/FractureEditor/Private/FractureToolPlaneCut.h
@@ -17,7 +17,7 @@ public:
 
 	UFracturePlaneCutSettings(const FObjectInitializer& ObjInit)
 		: Super(ObjInit)
-		, NumberPlanarCuts(3) {}
+		, NumberPlanarCuts(1) {}
 
 	/** Number of Clusters - Cluster Voronoi Method */
 	UPROPERTY(EditAnywhere, Category = PlaneCut, meta = (DisplayName = "Number of Cuts", UIMin = "1", UIMax = "20", ClampMin = "1"))
