GeometryCollectionComponent
Fracture Mode
-- moved ShowBoneColors from FractureMode to GeometryCollectionComponent to make it persistent.

diff --git a/Engine/Plugins/Experimental/ChaosEditor/Source/FractureEditor/Private/FractureEditorModeToolkit.cpp b/Engine/Plugins/Experimental/ChaosEditor/Source/FractureEditor/Private/FractureEditorModeToolkit.cpp
--- a/Engine/Plugins/Experimental/ChaosEditor/Source/FractureEditor/Private/FractureEditorModeToolkit.cpp
+++ b/Engine/Plugins/Experimental/ChaosEditor/Source/FractureEditor/Private/FractureEditorModeToolkit.cpp
@@ -98,24 +98,7 @@ FFractureViewSettingsCustomization::CustomizeDetails(IDetailLayoutBuilder& Detai
 	IDetailCategoryBuilder& ViewCategory = DetailBuilder.EditCategory("ViewSettings", FText::GetEmpty(), ECategoryPriority::TypeSpecific);
 
 	TSharedRef<IPropertyHandle> LevelProperty = DetailBuilder.GetProperty("FractureLevel");
-
-	ViewCategory.AddCustomRow(FText::GetEmpty())
-	.NameContent()
-	.HAlign(HAlign_Right)
-	[
-		SNew(STextBlock)
-		.TextStyle( &FCoreStyle::Get().GetWidgetStyle<FTextBlockStyle>( "SmallText" ) )
-		.Text(LOCTEXT("ShowBoneColors", "Show Bone Colors"))
-
-	]
-	.ValueContent()
-	[
-		SNew(SCheckBox)
-		.IsChecked_Lambda([this]() -> ECheckBoxState { return this->Toolkit->GetShowBoneColors() ? ECheckBoxState::Checked : ECheckBoxState::Unchecked; } )
-		.OnCheckStateChanged_Lambda([this] (ECheckBoxState CheckState) { this->Toolkit->OnSetShowBoneColors( CheckState == ECheckBoxState::Checked ); } )
-	];
-
-
+	
 	ViewCategory.AddProperty(LevelProperty)
 	.CustomWidget()
 	.NameContent()
@@ -548,13 +531,6 @@ void FFractureEditorModeToolkit::BindCommands()
 {
 	const FFractureEditorCommands& Commands = FFractureEditorCommands::Get();
 	
-	ToolkitCommands->MapAction(
-		Commands.ToggleShowBoneColors,
-		FExecuteAction::CreateSP(this, &FFractureEditorModeToolkit::OnToggleShowBoneColors),
-		FCanExecuteAction(),
-		FIsActionChecked::CreateSP(this, &FFractureEditorModeToolkit::GetShowBoneColors)
-	);
-
 	ToolkitCommands->MapAction(
 		Commands.ViewUpOneLevel,
 		FExecuteAction::CreateSP(this, &FFractureEditorModeToolkit::ViewUpOneLevel)//,
@@ -671,25 +647,6 @@ bool FFractureEditorModeToolkit::GetShowBoneColors() const
 	return false;
 }
 
-void FFractureEditorModeToolkit::OnToggleShowBoneColors()
-{
-	OnSetShowBoneColors(!GetShowBoneColors());
-}
-
-void FFractureEditorModeToolkit::OnSetShowBoneColors(bool NewValue)
-{
-	TSet<UGeometryCollectionComponent*> GeomCompSelection;
-	GetSelectedGeometryCollectionComponents(GeomCompSelection);
-	for (UGeometryCollectionComponent* Comp : GeomCompSelection)
-	{
-		FScopedColorEdit EditBoneColor = Comp->EditBoneSelection();
-		EditBoneColor.SetShowBoneColors(NewValue);
-		Comp->MarkRenderStateDirty();
-		Comp->MarkRenderDynamicDataDirty();
-	}
-	GCurrentLevelEditingViewportClient->Invalidate();
-}
-
 void FFractureEditorModeToolkit::OnSetExplodedViewValue(float NewValue)
 {
 	FScopedTransaction Transaction(LOCTEXT("SetExplodedViewValue", "Adjust Exploded View"));
diff --git a/Engine/Plugins/Experimental/ChaosEditor/Source/FractureEditor/Public/FractureEditorModeToolkit.h b/Engine/Plugins/Experimental/ChaosEditor/Source/FractureEditor/Public/FractureEditorModeToolkit.h
--- a/Engine/Plugins/Experimental/ChaosEditor/Source/FractureEditor/Public/FractureEditorModeToolkit.h
+++ b/Engine/Plugins/Experimental/ChaosEditor/Source/FractureEditor/Public/FractureEditorModeToolkit.h
@@ -107,8 +107,6 @@ public:
 	bool GetShowBoneColors() const;
 	void OnSetExplodedViewValue(float NewValue);
 	void OnSetLevelViewValue(int32 NewValue);
-	void OnToggleShowBoneColors();
-	void OnSetShowBoneColors(bool NewValue);
 
 	void OnExplodedViewValueChanged();
 	void OnLevelViewValueChanged();
diff --git a/Engine/Source/Runtime/Experimental/GeometryCollectionEngine/Private/GeometryCollection/GeometryCollectionComponent.cpp b/Engine/Source/Runtime/Experimental/GeometryCollectionEngine/Private/GeometryCollection/GeometryCollectionComponent.cpp
--- a/Engine/Source/Runtime/Experimental/GeometryCollectionEngine/Private/GeometryCollection/GeometryCollectionComponent.cpp
+++ b/Engine/Source/Runtime/Experimental/GeometryCollectionEngine/Private/GeometryCollection/GeometryCollectionComponent.cpp
@@ -230,11 +230,11 @@ UGeometryCollectionComponent::UGeometryCollectionComponent(const FObjectInitiali
 	, CachePlayback(false)
 	, bNotifyBreaks(false)
 	, bNotifyCollisions(false)
+	, bShowBoneColors(false)
 	, bEnableReplication(false)
 	, bEnableAbandonAfterLevel(false)
 	, ReplicationAbandonClusterLevel(0)
 	, bRenderStateDirty(true)
-	, bShowBoneColors(false)
 	, bEnableBoneSelection(false)
 	, ViewLevel(-1)
 	, NavmeshInvalidationTimeSliceIndex(0)
@@ -858,6 +858,19 @@ void UGeometryCollectionComponent::InitializeComponent()
 	}
 }
 
+void UGeometryCollectionComponent::PostEditChangeProperty(FPropertyChangedEvent& PropertyChangedEvent)
+{
+	Super::PostEditChangeProperty(PropertyChangedEvent);
+
+	if (PropertyChangedEvent.Property && PropertyChangedEvent.Property->GetFName() == GET_MEMBER_NAME_CHECKED(UGeometryCollectionComponent, bShowBoneColors))
+	{
+		FScopedColorEdit EditBoneColor = EditBoneSelection();
+		EditBoneColor.SetShowBoneColors(bShowBoneColors);
+		MarkRenderStateDirty();
+		MarkRenderDynamicDataDirty();
+	}
+}
+
 static void DispatchGeometryCollectionBreakEvent(const FChaosBreakEvent& Event)
 {
 	if (UGeometryCollectionComponent* const GC = Cast<UGeometryCollectionComponent>(Event.Component))
diff --git a/Engine/Source/Runtime/Experimental/GeometryCollectionEngine/Public/GeometryCollection/GeometryCollectionComponent.h b/Engine/Source/Runtime/Experimental/GeometryCollectionEngine/Public/GeometryCollection/GeometryCollectionComponent.h
--- a/Engine/Source/Runtime/Experimental/GeometryCollectionEngine/Public/GeometryCollection/GeometryCollectionComponent.h
+++ b/Engine/Source/Runtime/Experimental/GeometryCollectionEngine/Public/GeometryCollection/GeometryCollectionComponent.h
@@ -341,6 +341,7 @@ public:
 	virtual void EndPlay(const EEndPlayReason::Type ReasonEnd) override;
 	virtual void GetLifetimeReplicatedProps(TArray<FLifetimeProperty>& OutLifetimeProps) const override;
 	virtual void InitializeComponent() override;
+	virtual void PostEditChangeProperty(FPropertyChangedEvent& PropertyChangedEvent) override;
 	//~ Begin UActorComponent Interface. 
 
 
@@ -652,6 +653,10 @@ protected:
 	UPROPERTY(EditAnywhere, BlueprintReadOnly, Category = "ChaosPhysics|General")
 	bool bNotifyCollisions;
 
+	/** Display Bone Colors instead of assigned materials */
+	UPROPERTY(EditAnywhere, BlueprintReadWrite, Category = "ChaosPhysics|General")
+	bool bShowBoneColors;
+
 	/** Populate the static geometry structures for the render thread. */
 	void InitConstantData(FGeometryCollectionConstantData* ConstantData) const;
 
@@ -714,7 +719,6 @@ private:
 	void NetAbandonCluster(int32 TransformIndex);
 
 	bool bRenderStateDirty;
-	bool bShowBoneColors;
 	bool bEnableBoneSelection;
 	int ViewLevel;
 
