Fracture Mode
-- Bone Color does not change when exiting Fracture Mode or changing selection.

diff --git a/Engine/Plugins/Experimental/ChaosEditor/Source/FractureEditor/Private/FractureEditorMode.cpp b/Engine/Plugins/Experimental/ChaosEditor/Source/FractureEditor/Private/FractureEditorMode.cpp
--- a/Engine/Plugins/Experimental/ChaosEditor/Source/FractureEditor/Private/FractureEditorMode.cpp
+++ b/Engine/Plugins/Experimental/ChaosEditor/Source/FractureEditor/Private/FractureEditorMode.cpp
@@ -338,7 +338,6 @@ void FFractureEditorMode::OnActorSelectionChanged(const TArray<UObject*>& NewSel
 
 			FScopedColorEdit ShowBoneColorsEdit(ExistingSelection);
 			ShowBoneColorsEdit.SetEnableBoneSelection(false);
-			ShowBoneColorsEdit.SetShowBoneColors(false);
 			
 			ExistingSelection->MarkRenderStateDirty();
 		}
diff --git a/Engine/Plugins/Experimental/ChaosEditor/Source/FractureEditor/Private/FractureEditorModeToolkit.cpp b/Engine/Plugins/Experimental/ChaosEditor/Source/FractureEditor/Private/FractureEditorModeToolkit.cpp
--- a/Engine/Plugins/Experimental/ChaosEditor/Source/FractureEditor/Private/FractureEditorModeToolkit.cpp
+++ b/Engine/Plugins/Experimental/ChaosEditor/Source/FractureEditor/Private/FractureEditorModeToolkit.cpp
@@ -632,21 +632,6 @@ int32 FFractureEditorModeToolkit::GetLevelViewValue() const
 	return FractureSettings->FractureLevel;
 }
 
-bool FFractureEditorModeToolkit::GetShowBoneColors() const
-{
-	
-	TSet<UGeometryCollectionComponent*> GeomCompSelection;
-	GetSelectedGeometryCollectionComponents(GeomCompSelection);
-	if (GeomCompSelection.Num() > 0)
-	{
-		UGeometryCollectionComponent* Comp = GeomCompSelection.Array()[0];
-		FScopedColorEdit EditBoneColor = Comp->EditBoneSelection();
-		return EditBoneColor.GetShowBoneColors();
-	}
-
-	return false;
-}
-
 void FFractureEditorModeToolkit::OnSetExplodedViewValue(float NewValue)
 {
 	FScopedTransaction Transaction(LOCTEXT("SetExplodedViewValue", "Adjust Exploded View"));
diff --git a/Engine/Plugins/Experimental/ChaosEditor/Source/FractureEditor/Public/FractureEditorModeToolkit.h b/Engine/Plugins/Experimental/ChaosEditor/Source/FractureEditor/Public/FractureEditorModeToolkit.h
--- a/Engine/Plugins/Experimental/ChaosEditor/Source/FractureEditor/Public/FractureEditorModeToolkit.h
+++ b/Engine/Plugins/Experimental/ChaosEditor/Source/FractureEditor/Public/FractureEditorModeToolkit.h
@@ -104,7 +104,6 @@ public:
 	// View Settings
 	float GetExplodedViewValue() const;
 	int32 GetLevelViewValue() const;
-	bool GetShowBoneColors() const;
 	void OnSetExplodedViewValue(float NewValue);
 	void OnSetLevelViewValue(int32 NewValue);
 
