Fracture Mode
-- Optimized Voronoi Auto Cluster

diff --git a/Engine/Plugins/Experimental/ChaosEditor/Source/FractureEditor/Private/FractureToolAutoCluster.cpp b/Engine/Plugins/Experimental/ChaosEditor/Source/FractureEditor/Private/FractureToolAutoCluster.cpp
--- a/Engine/Plugins/Experimental/ChaosEditor/Source/FractureEditor/Private/FractureToolAutoCluster.cpp
+++ b/Engine/Plugins/Experimental/ChaosEditor/Source/FractureEditor/Private/FractureToolAutoCluster.cpp
@@ -71,6 +71,7 @@ void UFractureToolAutoCluster::Execute(TWeakPtr<FFractureEditorModeToolkit> InTo
 		for (FFractureToolContext& Context : Contexts)
 		{
 			Context.ConvertSelectionToClusterNodes();
+			FGeometryCollection* GeometryCollection = Context.GetGeometryCollection().Get();
 
 			if (AutoClusterSettings->AutoClusterMode < EFractureAutoClusterMode::Voronoi)
 			{
@@ -80,30 +81,39 @@ void UFractureToolAutoCluster::Execute(TWeakPtr<FFractureEditorModeToolkit> InTo
 			{
 				for (const int32 ClusterIndex : Context.GetSelection())
 				{
-					FVoronoiPartitioner VoronoiPartition(Context.GetGeometryCollection().Get(), ClusterIndex);
+					FVoronoiPartitioner VoronoiPartition(GeometryCollection, ClusterIndex);
 					VoronoiPartition.KMeansPartition(AutoClusterSettings->SiteCount);
 
 					if (AutoClusterSettings->bEnforceConnectivity)
 					{
-						GenerateProximityIfNecessary(Context.GetGeometryCollection().Get());
-						VoronoiPartition.SplitDisconnectedPartitions(Context.GetGeometryCollection().Get());
+						GenerateProximityIfNecessary(GeometryCollection);
+						VoronoiPartition.SplitDisconnectedPartitions(GeometryCollection);
 					}
 					
 					int32 PartitionCount = VoronoiPartition.GetPartitionCount();
-					for (int32 NewClusterIndex = 0; NewClusterIndex < VoronoiPartition.GetPartitionCount(); ++NewClusterIndex)
+					int32 NewClusterIndexStart = GeometryCollection->AddElements(PartitionCount, FGeometryCollection::TransformGroup);
+
+					for (int32 Index = 0; Index < PartitionCount; ++Index)
 					{
-						TArray<int32> NewCluster = VoronoiPartition.GetPartition(NewClusterIndex);
-						if (NewCluster.Num() > 0)
-						{
-							FGeometryCollectionClusteringUtility::ClusterBonesUnderNewNode(Context.GetGeometryCollection().Get(), NewCluster[0], NewCluster, true);
-						}
+						
+						TArray<int32> NewCluster = VoronoiPartition.GetPartition(Index);
+
+						int32 NewClusterIndex = NewClusterIndexStart + Index;
+						GeometryCollection->Parent[NewClusterIndex] = ClusterIndex;
+						GeometryCollection->Children[ClusterIndex].Add(NewClusterIndex);
+						GeometryCollection->BoneName[NewClusterIndex] = "ClusterBone";
+						GeometryCollection->Children[NewClusterIndex] = TSet<int32>(NewCluster);
+						GeometryCollection->SimulationType[NewClusterIndex] = FGeometryCollection::ESimulationTypes::FST_Clustered;
+						GeometryCollection->Transform[NewClusterIndex] = FTransform::Identity;
+						GeometryCollectionAlgo::ParentTransforms(GeometryCollection, NewClusterIndex, NewCluster);
 					}
+					FGeometryCollectionClusteringUtility::UpdateHierarchyLevelOfChildren(GeometryCollection, ClusterIndex);
+					FGeometryCollectionClusteringUtility::RecursivelyUpdateChildBoneNames(ClusterIndex, GeometryCollection->Children, GeometryCollection->BoneName);
+					FGeometryCollectionClusteringUtility::ValidateResults(GeometryCollection);
 				}
 			}
-
 			Refresh(Context, Toolkit);
 		}
-
 		SetOutlinerComponents(Contexts, Toolkit);
 	}
 }
@@ -237,10 +247,10 @@ void FVoronoiPartitioner::CollectConnections(const FGeometryCollection* Geometry
 void FVoronoiPartitioner::GenerateCentroids(const FGeometryCollection* GeometryCollection)
 {
 	Centroids.SetNum(TransformIndices.Num());
-	for (int32 Index = 0; Index < Centroids.Num(); ++Index)
-	{
-		Centroids[Index] = GenerateCentroid(GeometryCollection, TransformIndices[Index]);
-	}
+	ParallelFor(TransformIndices.Num(), [this, GeometryCollection](int32 Index)
+		{
+			Centroids[Index] = GenerateCentroid(GeometryCollection, TransformIndices[Index]);
+		});
 }
 
 FVector FVoronoiPartitioner::GenerateCentroid(const FGeometryCollection* GeometryCollection, int32 TransformIndex) const
