Fracture Mode
-- re-instated shift-b to toggle bone colors in fracture mode.

diff --git a/Engine/Plugins/Experimental/ChaosEditor/Source/FractureEditor/Private/FractureEditorModeToolkit.cpp b/Engine/Plugins/Experimental/ChaosEditor/Source/FractureEditor/Private/FractureEditorModeToolkit.cpp
--- a/Engine/Plugins/Experimental/ChaosEditor/Source/FractureEditor/Private/FractureEditorModeToolkit.cpp
+++ b/Engine/Plugins/Experimental/ChaosEditor/Source/FractureEditor/Private/FractureEditorModeToolkit.cpp
@@ -531,6 +531,11 @@ void FFractureEditorModeToolkit::BindCommands()
 {
 	const FFractureEditorCommands& Commands = FFractureEditorCommands::Get();
 	
+	ToolkitCommands->MapAction(
+		Commands.ToggleShowBoneColors,
+		FExecuteAction::CreateSP(this, &FFractureEditorModeToolkit::ToggleShowBoneColors)//,
+	);
+	
 	ToolkitCommands->MapAction(
 		Commands.ViewUpOneLevel,
 		FExecuteAction::CreateSP(this, &FFractureEditorModeToolkit::ViewUpOneLevel)//,
@@ -754,6 +759,17 @@ void FFractureEditorModeToolkit::OnLevelViewValueChanged()
 	GCurrentLevelEditingViewportClient->Invalidate();
 }
 
+void FFractureEditorModeToolkit::ToggleShowBoneColors()
+{
+	TSet<UGeometryCollectionComponent*> GeomCompSelection;
+	GetSelectedGeometryCollectionComponents(GeomCompSelection);
+
+	for (UGeometryCollectionComponent* Comp : GeomCompSelection)
+	{
+		Comp->SetShowBoneColors(!Comp->GetShowBoneColors());
+	}
+}
+
 void FFractureEditorModeToolkit::ViewUpOneLevel()
 {
 	int32 CountMax = GetLevelCount() + 1;
diff --git a/Engine/Plugins/Experimental/ChaosEditor/Source/FractureEditor/Public/FractureEditorModeToolkit.h b/Engine/Plugins/Experimental/ChaosEditor/Source/FractureEditor/Public/FractureEditorModeToolkit.h
--- a/Engine/Plugins/Experimental/ChaosEditor/Source/FractureEditor/Public/FractureEditorModeToolkit.h
+++ b/Engine/Plugins/Experimental/ChaosEditor/Source/FractureEditor/Public/FractureEditorModeToolkit.h
@@ -116,6 +116,8 @@ public:
 	TSharedRef<SWidget> GetLevelViewMenuContent(TSharedRef<IPropertyHandle> PropertyHandle);
 	TSharedRef<SWidget> GetViewMenuContent();
 
+	void ToggleShowBoneColors();
+
 	void ViewUpOneLevel();
 	void ViewDownOneLevel();
 
diff --git a/Engine/Source/Runtime/Experimental/GeometryCollectionEngine/Public/GeometryCollection/GeometryCollectionComponent.h b/Engine/Source/Runtime/Experimental/GeometryCollectionEngine/Public/GeometryCollection/GeometryCollectionComponent.h
--- a/Engine/Source/Runtime/Experimental/GeometryCollectionEngine/Public/GeometryCollection/GeometryCollectionComponent.h
+++ b/Engine/Source/Runtime/Experimental/GeometryCollectionEngine/Public/GeometryCollection/GeometryCollectionComponent.h
@@ -561,7 +561,7 @@ public:
 	*/
 	/* ---------------------------------------------------------------------------------------- */
 	
-	void SetShowBoneColors(bool ShowBoneColorsIn);
+	void SetShowBoneColors(bool ShowBoneColorsIn) { bShowBoneColors = ShowBoneColorsIn; }
 	bool GetShowBoneColors() const { return bShowBoneColors; }
 	bool GetEnableBoneSelection() const { return bEnableBoneSelection; }
 	
