Fracture Mode
-- fixed refresh bug with shift-b hotkey

diff --git a/Engine/Plugins/Experimental/ChaosEditor/Source/FractureEditor/Private/FractureEditorModeToolkit.cpp b/Engine/Plugins/Experimental/ChaosEditor/Source/FractureEditor/Private/FractureEditorModeToolkit.cpp
--- a/Engine/Plugins/Experimental/ChaosEditor/Source/FractureEditor/Private/FractureEditorModeToolkit.cpp
+++ b/Engine/Plugins/Experimental/ChaosEditor/Source/FractureEditor/Private/FractureEditorModeToolkit.cpp
@@ -767,6 +767,8 @@ void FFractureEditorModeToolkit::ToggleShowBoneColors()
 	for (UGeometryCollectionComponent* Comp : GeomCompSelection)
 	{
 		Comp->SetShowBoneColors(!Comp->GetShowBoneColors());
+		Comp->MarkRenderStateDirty();
+		Comp->MarkRenderDynamicDataDirty();
 	}
 }
 
