Fracture Tool
-- stop crashing when auto clustering an empty cluster

diff --git a/Engine/Plugins/Experimental/ChaosEditor/Source/FractureEditor/Private/FractureToolAutoCluster.cpp b/Engine/Plugins/Experimental/ChaosEditor/Source/FractureEditor/Private/FractureToolAutoCluster.cpp
--- a/Engine/Plugins/Experimental/ChaosEditor/Source/FractureEditor/Private/FractureToolAutoCluster.cpp
+++ b/Engine/Plugins/Experimental/ChaosEditor/Source/FractureEditor/Private/FractureToolAutoCluster.cpp
@@ -288,7 +288,13 @@ FBox FVoronoiPartitioner::GenerateBounds(const FGeometryCollection* GeometryColl
 	else // recurse the cluster
 	{
 		const TArray<int32> Children = GeometryCollection->Children[TransformIndex].Array();
-		check(Children.Num() > 0);
+
+		// Empty clusters are problematic
+		if (Children.Num() == 0)
+		{
+			return FBox();
+		}
+
 		FBox Bounds = GenerateBounds(GeometryCollection, Children[0]);
 		for (int32 ChildIndex = 1; ChildIndex < Children.Num(); ++ChildIndex)
 		{
