Add a geocoll->staticmesh convert tool to the fracture editor (under the 'fix' category for now)

diff --git a/Engine/Plugins/Experimental/ChaosEditor/Source/FractureEditor/FractureEditor.Build.cs b/Engine/Plugins/Experimental/ChaosEditor/Source/FractureEditor/FractureEditor.Build.cs
--- a/Engine/Plugins/Experimental/ChaosEditor/Source/FractureEditor/FractureEditor.Build.cs
+++ b/Engine/Plugins/Experimental/ChaosEditor/Source/FractureEditor/FractureEditor.Build.cs
@@ -48,6 +48,8 @@ public class FractureEditor : ModuleRules
                 "GeometryCollectionEngine",
                 "GeometryCollectionEditor",
 				"ModelingComponents",
+				"MeshDescription",
+				"StaticMeshDescription",
 				"PlanarCut",
 				"Chaos"
 
diff --git a/Engine/Plugins/Experimental/ChaosEditor/Source/FractureEditor/Private/FractureEditorCommands.h b/Engine/Plugins/Experimental/ChaosEditor/Source/FractureEditor/Private/FractureEditorCommands.h
--- a/Engine/Plugins/Experimental/ChaosEditor/Source/FractureEditor/Private/FractureEditorCommands.h
+++ b/Engine/Plugins/Experimental/ChaosEditor/Source/FractureEditor/Private/FractureEditorCommands.h
@@ -62,6 +62,7 @@ class FFractureEditorCommands : public TCommands<FFractureEditorCommands>
 		// Cleanup Commands
 		TSharedPtr< FUICommandInfo > RecomputeNormals;
 		TSharedPtr< FUICommandInfo > Resample;
+		TSharedPtr< FUICommandInfo > ConvertToMesh;
 
 		// Property Commands
 		TSharedPtr< FUICommandInfo > SetInitialDynamicState;
diff --git a/Engine/Plugins/Experimental/ChaosEditor/Source/FractureEditor/Private/FractureEditorModeToolkit.cpp b/Engine/Plugins/Experimental/ChaosEditor/Source/FractureEditor/Private/FractureEditorModeToolkit.cpp
--- a/Engine/Plugins/Experimental/ChaosEditor/Source/FractureEditor/Private/FractureEditorModeToolkit.cpp
+++ b/Engine/Plugins/Experimental/ChaosEditor/Source/FractureEditor/Private/FractureEditorModeToolkit.cpp
@@ -520,6 +520,7 @@ void FFractureEditorModeToolkit::BuildToolPalette(FName PaletteIndex, class FToo
 	{
 		ToolbarBuilder.AddToolBarButton(Commands.RecomputeNormals);
 		ToolbarBuilder.AddToolBarButton(Commands.Resample);
+		ToolbarBuilder.AddToolBarButton(Commands.ConvertToMesh);
 	}
 	else if (PaletteIndex == TEXT("Properties"))
 	{
diff --git a/Engine/Plugins/Experimental/ChaosEditor/Source/FractureEditor/Private/FractureToolConvert.cpp b/Engine/Plugins/Experimental/ChaosEditor/Source/FractureEditor/Private/FractureToolConvert.cpp
new file mode 100644
--- /dev/null
+++ b/Engine/Plugins/Experimental/ChaosEditor/Source/FractureEditor/Private/FractureToolConvert.cpp
@@ -0,0 +1,332 @@
+// Copyright Epic Games, Inc. All Rights Reserved.
+
+#include "FractureToolConvert.h"
+
+#include "FractureToolContext.h"
+
+#include "GeometryCollection/GeometryCollectionObject.h"
+#include "GeometryCollection/GeometryCollectionComponent.h"
+#include "GeometryCollection/GeometryCollectionAlgo.h"
+
+#include "AssetToolsModule.h"
+#include "AssetRegistryModule.h"
+// for content-browser things
+#include "ContentBrowserModule.h"
+#include "IContentBrowserSingleton.h"
+
+#include "Engine/StaticMesh.h"
+#include "Engine/StaticMeshActor.h"
+#include "Components/StaticMeshComponent.h"
+#include "StaticMeshAttributes.h"
+#include "Engine/CollisionProfile.h"
+#include "PhysicsEngine/BodySetup.h"
+
+#include "FileHelpers.h" // TODO: rm if we aren't keeping the prompt for checkout and save fn
+
+#include "PlanarCut.h"
+
+#include "Misc/ScopedSlowTask.h"
+
+#define LOCTEXT_NAMESPACE "FractureToolConvert"
+
+
+UFractureToolConvert::UFractureToolConvert(const FObjectInitializer& ObjInit)
+	: Super(ObjInit)
+{
+	ConvertSettings = NewObject<UFractureConvertSettings>(GetTransientPackage(), UFractureConvertSettings::StaticClass());
+	ConvertSettings->OwnerTool = this;
+}
+
+bool UFractureToolConvert::CanExecute() const
+{
+	if (!IsGeometryCollectionSelected())
+	{
+		return false;
+	}
+
+	return true;
+}
+
+FText UFractureToolConvert::GetDisplayText() const
+{
+	return FText(LOCTEXT("FractureToolConvert", "Convert to Static Mesh"));
+}
+
+FText UFractureToolConvert::GetTooltipText() const
+{
+	return FText(LOCTEXT("FractureToolConvertTooltip", "This converts fracture geometry to static meshes."));
+}
+
+FSlateIcon UFractureToolConvert::GetToolIcon() const
+{
+	return FSlateIcon("FractureEditorStyle", "FractureEditor.Convert");
+}
+
+void UFractureToolConvert::RegisterUICommand(FFractureEditorCommands* BindingContext)
+{
+	UI_COMMAND_EXT(BindingContext, UICommandInfo, "ToMesh", "ToMesh", "Convert geometry to static mesh", EUserInterfaceActionType::ToggleButton, FInputChord());
+	BindingContext->ConvertToMesh = UICommandInfo;
+}
+
+TArray<UObject*> UFractureToolConvert::GetSettingsObjects() const
+{
+	TArray<UObject*> Settings;
+	Settings.Add(ConvertSettings);
+	return Settings;
+}
+
+void UFractureToolConvert::FractureContextChanged()
+{
+}
+
+void UFractureToolConvert::Render(const FSceneView* View, FViewport* Viewport, FPrimitiveDrawInterface* PDI)
+{
+}
+
+void UFractureToolConvert::PostEditChangeChainProperty(struct FPropertyChangedChainEvent& PropertyChangedEvent)
+{
+	// update any cached data 
+}
+
+TArray<FFractureToolContext> UFractureToolConvert::GetFractureToolContexts() const
+{
+	TArray<FFractureToolContext> Contexts;
+
+	// A context is gathered for each selected GeometryCollection component, or for each individual bone if Group Fracture is not used.
+	TSet<UGeometryCollectionComponent*> GeomCompSelection;
+	GetSelectedGeometryCollectionComponents(GeomCompSelection);
+
+	for (UGeometryCollectionComponent* GeometryCollectionComponent : GeomCompSelection)
+	{
+		// Generate a context for each selected node
+		FFractureToolContext FullSelection(GeometryCollectionComponent);
+		FullSelection.ConvertSelectionToRigidNodes();
+		
+		// Update global transforms and bounds -- TODO: pull this bounds update out to a shared function?
+		const TManagedArray<FTransform>& Transform = FullSelection.GetGeometryCollection()->Transform;
+		const TManagedArray<int32>& TransformToGeometryIndex = FullSelection.GetGeometryCollection()->TransformToGeometryIndex;
+		const TManagedArray<FBox>& BoundingBoxes = FullSelection.GetGeometryCollection()->BoundingBox;
+
+		TArray<FTransform> Transforms;
+		GeometryCollectionAlgo::GlobalMatrices(Transform, FullSelection.GetGeometryCollection()->Parent, Transforms);
+
+		FBox Bounds(ForceInit);
+		for (int32 BoneIndex : FullSelection.GetSelection())
+		{
+			int32 GeometryIndex = TransformToGeometryIndex[BoneIndex];
+			if (GeometryIndex > INDEX_NONE)
+			{
+				FBox BoneBound = BoundingBoxes[GeometryIndex].TransformBy(Transforms[BoneIndex]);
+				Bounds += BoneBound;
+			}
+		}
+		FullSelection.SetBounds(Bounds);
+
+		Contexts.Add(FullSelection);
+	}
+
+	return Contexts;
+}
+
+namespace UE
+{
+namespace FractureToolConvertInternal
+{
+// Create + position a single static mesh asset in the provided asset package
+AActor* CreateMeshAsset(FGeometryCollection& Collection, const FTransform& CollectionToWorld, const TConstArrayView<UMaterialInterface*>& Materials, TConstArrayView<int32> Bones, UPackage* AssetPackage, FString UniqueAssetName, bool bCenterPivot)
+{
+	UWorld* TargetWorld = GEditor->GetEditorWorldContext().World();
+	check(TargetWorld);
+
+	// create new UStaticMesh object
+	EObjectFlags Flags = EObjectFlags::RF_Public | EObjectFlags::RF_Standalone;
+	UStaticMesh* NewStaticMesh = NewObject<UStaticMesh>(AssetPackage, *UniqueAssetName, Flags);
+
+	// initialize the LOD 0 MeshDescription
+	NewStaticMesh->SetNumSourceModels(1);
+	// normals and tangents should carry over from the geometry collection
+	NewStaticMesh->GetSourceModel(0).BuildSettings.bRecomputeNormals = false;
+	NewStaticMesh->GetSourceModel(0).BuildSettings.bRecomputeTangents = false;
+	FMeshDescription* OutputMeshDescription = NewStaticMesh->CreateMeshDescription(0);
+
+	NewStaticMesh->CreateBodySetup();
+	NewStaticMesh->GetBodySetup()->CollisionTraceFlag = ECollisionTraceFlag::CTF_UseComplexAsSimple;
+
+
+	FTransform BonesToCollection;
+	ConvertToMeshDescription(*OutputMeshDescription, BonesToCollection, bCenterPivot, Collection, Bones);
+
+	// add a material slot. Must always have one material slot.
+	int AddMaterialCount = FMath::Max(1, Materials.Num());
+	for (int MatIdx = 0; MatIdx < AddMaterialCount; MatIdx++)
+	{
+		NewStaticMesh->GetStaticMaterials().Add(FStaticMaterial());
+	}
+	for (int MatIdx = 0; MatIdx < Materials.Num(); MatIdx++)
+	{
+		NewStaticMesh->SetMaterial(MatIdx, Materials[MatIdx]);
+	}
+
+	NewStaticMesh->CommitMeshDescription(0);
+
+	NewStaticMesh->GetBodySetup()->CollisionTraceFlag = ECollisionTraceFlag::CTF_UseComplexAsSimple;
+
+
+	// create new actor
+	FRotator Rotation(0.0f, 0.0f, 0.0f);
+	FActorSpawnParameters SpawnInfo;
+	AStaticMeshActor* NewActor = TargetWorld->SpawnActor<AStaticMeshActor>(FVector::ZeroVector, Rotation, SpawnInfo);
+	NewActor->SetActorLabel(*UniqueAssetName);
+	NewActor->GetStaticMeshComponent()->SetStaticMesh(NewStaticMesh);
+
+	// if we don't do this, world traces don't hit the mesh
+	NewActor->GetStaticMeshComponent()->RecreatePhysicsState();
+
+	NewActor->GetStaticMeshComponent()->SetWorldTransform(CollectionToWorld * BonesToCollection);
+
+	for (int MatIdx = 0; MatIdx < Materials.Num(); MatIdx++)
+	{
+		NewActor->GetStaticMeshComponent()->SetMaterial(MatIdx, Materials[MatIdx]);
+	}
+
+	NewActor->MarkComponentsRenderStateDirty();
+
+	GEditor->SelectActor(NewActor, true, false, true, false);
+
+	NewStaticMesh->PostEditChange();
+
+	AssetPackage->MarkPackageDirty();
+	FAssetRegistryModule::AssetCreated(AssetPackage);
+
+	return NewActor;
+}
+
+// Convert and save all the requested static mesh assets
+bool ConvertAndSaveMeshes(FGeometryCollection& Collection, const FTransform& CollectionToWorld, const TConstArrayView<UMaterialInterface*>& Materials, TConstArrayView<int32> Bones, FString ObjectBaseName, const UObject* RelativeToAsset, bool bPromptToSave, bool bSaveCombined, bool bCenterPivots)
+{
+	check(RelativeToAsset);
+
+	// find path to reference asset
+	UPackage* AssetOuterPackage = CastChecked<UPackage>(RelativeToAsset->GetOuter());
+	FString AssetPackageName = AssetOuterPackage->GetName();
+	FString PackageFolderPath = FPackageName::GetLongPackagePath(AssetPackageName);
+
+	// Show the modal dialog and then get the path/name.
+	// If the user cancels, conversion is skipped
+	if (bPromptToSave)
+	{
+		IContentBrowserSingleton& ContentBrowser = FModuleManager::LoadModuleChecked<FContentBrowserModule>("ContentBrowser").Get();
+
+		FString UseDefaultAssetName = ObjectBaseName;
+		FString CurrentPath = PackageFolderPath;
+
+		FSaveAssetDialogConfig Config;
+		Config.ExistingAssetPolicy = ESaveAssetDialogExistingAssetPolicy::AllowButWarn;
+		Config.DefaultAssetName = UseDefaultAssetName;
+		Config.DialogTitleOverride = bSaveCombined ? 
+			LOCTEXT("GenerateBoneMeshAssetPathDialogWarningCombine", "Choose Folder Path and Base Name for Converted Mesh.") :
+			LOCTEXT("GenerateBoneMeshAssetPathDialogWarningPerBone", "Choose Folder Path and Base Name for Converted Meshes. Bone index will be appended to each mesh.");
+		Config.DefaultPath = CurrentPath;
+		FString SelectedPath = ContentBrowser.CreateModalSaveAssetDialog(Config);
+
+		if (SelectedPath.IsEmpty() == false)
+		{
+			PackageFolderPath = FPaths::GetPath(SelectedPath);
+			ObjectBaseName = FPaths::GetBaseFilename(SelectedPath, true);
+		}
+		else
+		{
+			return false;
+		}
+	}
+
+	auto MakeUniquePackage = [&PackageFolderPath](FString BaseName, FString& UniqueAssetNameOut) -> UPackage*
+	{
+		// create new package
+		FString UniquePackageName;
+
+		FAssetToolsModule& AssetToolsModule = FModuleManager::LoadModuleChecked<FAssetToolsModule>("AssetTools");
+		AssetToolsModule.Get().CreateUniqueAssetName(PackageFolderPath + TEXT("/") + BaseName, TEXT(""), UniquePackageName, UniqueAssetNameOut);
+		return CreatePackage(*UniquePackageName);
+	};
+
+	FScopedSlowTask ConvertTask(Bones.Num(), LOCTEXT("StartingConvert", "Converting geometry to static mesh"));
+	ConvertTask.MakeDialog();
+
+	GEditor->SelectNone(false, true, false);
+
+	TArray<UPackage*> SavePackage;
+	if (bSaveCombined)
+	{
+		ConvertTask.EnterProgressFrame(Bones.Num());
+		FString UniqueAssetName;
+		UPackage* AssetPackage = MakeUniquePackage(ObjectBaseName, UniqueAssetName);
+		SavePackage.Add(AssetPackage);
+		CreateMeshAsset(Collection, CollectionToWorld, Materials, Bones, AssetPackage, UniqueAssetName, bCenterPivots);
+	}
+	else
+	{
+		for (int SelIdx = 0; SelIdx < Bones.Num(); SelIdx++)
+		{
+			ConvertTask.EnterProgressFrame(1);
+
+			int32 UseBone = Bones[SelIdx];
+
+			FString UseBaseName = FString::Printf(TEXT("%s_%d_"), *ObjectBaseName, UseBone);
+
+			FString UniqueAssetName;
+			UPackage* AssetPackage = MakeUniquePackage(UseBaseName, UniqueAssetName);
+			SavePackage.Add(AssetPackage);
+
+			TArrayView<const int32> SingleBoneView(&UseBone, 1);
+
+			CreateMeshAsset(Collection, CollectionToWorld, Materials, SingleBoneView, AssetPackage, UniqueAssetName, bCenterPivots);
+		}
+	}
+
+	GEditor->NoteSelectionChange(true);
+
+	FEditorFileUtils::PromptForCheckoutAndSave(SavePackage, true, true);
+
+	return true;
+}
+}} // namespace UE::FractureToolConvertInternal
+
+
+int32 UFractureToolConvert::ExecuteFracture(const FFractureToolContext& FractureContext)
+{
+	if (FractureContext.IsValid())
+	{
+		FGeometryCollection& Collection = *FractureContext.GetGeometryCollection();
+
+		FMeshDescription Mesh;
+		FTransform MeshTransform;
+
+		TArray<TObjectPtr<class UMaterialInterface>>& OverrideMaterials = FractureContext.GetGeometryCollectionComponent()->OverrideMaterials;
+		TArray<TObjectPtr<class UMaterialInterface>>& MainMaterials = FractureContext.GetFracturedGeometryCollection()->Materials;
+		TArray<UMaterialInterface*> Materials;
+		Materials.Reserve(MainMaterials.Num());
+		for (int MatIdx = 0; MatIdx < MainMaterials.Num(); MatIdx++)
+		{
+			if (MatIdx < OverrideMaterials.Num() && !OverrideMaterials[MatIdx].IsNull())
+			{
+				Materials.Add(OverrideMaterials[MatIdx].Get());
+			}
+			else
+			{
+				Materials.Add(MainMaterials[MatIdx].Get());
+			}
+		}
+
+		// choose default mesh name based on corresponding geometry collection name
+		FString BaseName = FString::Printf(TEXT("%s_SM"), *FractureContext.GetFracturedGeometryCollection()->GetName());
+		
+		UE::FractureToolConvertInternal::ConvertAndSaveMeshes(Collection, FractureContext.GetTransform(), Materials,
+			FractureContext.GetSelection(), BaseName, FractureContext.GetFracturedGeometryCollection(), 
+			ConvertSettings->bPromptForBaseName, !ConvertSettings->bPerBone, ConvertSettings->bCenterPivots);
+	}
+
+	return INDEX_NONE;
+}
+
+#undef LOCTEXT_NAMESPACE
diff --git a/Engine/Plugins/Experimental/ChaosEditor/Source/FractureEditor/Private/FractureToolConvert.h b/Engine/Plugins/Experimental/ChaosEditor/Source/FractureEditor/Private/FractureToolConvert.h
new file mode 100644
--- /dev/null
+++ b/Engine/Plugins/Experimental/ChaosEditor/Source/FractureEditor/Private/FractureToolConvert.h
@@ -0,0 +1,74 @@
+// Copyright Epic Games, Inc. All Rights Reserved.
+
+#pragma once
+
+#include "FractureTool.h"
+
+#include "FractureToolConvert.generated.h"
+
+class FFractureToolContext;
+struct FMeshDescription;
+
+/** Settings related to geometry collection -> static mesh conversion **/
+UCLASS(config = EditorPerProjectUserSettings)
+class UFractureConvertSettings : public UFractureToolSettings
+{
+public:
+
+	GENERATED_BODY()
+
+	UFractureConvertSettings(const FObjectInitializer& ObjInit)
+		: Super(ObjInit)
+	{}
+
+	/** Whether to prompt user for a location and base name for the generated meshes, or automatically place them next to the source geometry collections */
+	UPROPERTY(EditAnywhere, Category = AssetSettings)
+	bool bPromptForBaseName = true;
+
+	/** Whether to generate a separate mesh for each bone, or one combined mesh */
+	UPROPERTY(EditAnywhere, Category = AssetSettings)
+	bool bPerBone = true;
+
+	/** Whether to center the pivot for each mesh, or use the pivot from the geometry collection */
+	UPROPERTY(EditAnywhere, Category = GeometrySettings)
+	bool bCenterPivots = true;
+};
+
+
+UCLASS(DisplayName = "Convert Tool", Category = "FractureTools")
+class UFractureToolConvert : public UFractureModalTool
+{
+public:
+	GENERATED_BODY()
+
+	UFractureToolConvert(const FObjectInitializer& ObjInit);
+
+	///
+	/// UFractureModalTool Interface
+	///
+
+	/** This is the Text that will appear on the tool button to execute the tool **/
+	virtual FText GetDisplayText() const override;
+	virtual FText GetTooltipText() const override;
+	/** This is the Text that will appear on the button to execute the fracture **/
+	virtual FText GetApplyText() const override { return FText(NSLOCTEXT("FractureToolConvert", "ExecuteConvert", "Convert")); }
+	virtual FSlateIcon GetToolIcon() const override;
+	virtual void RegisterUICommand(FFractureEditorCommands* BindingContext) override;
+	virtual TArray<UObject*> GetSettingsObjects() const override;
+	virtual void FractureContextChanged() override;
+	virtual void Render(const FSceneView* View, FViewport* Viewport, FPrimitiveDrawInterface* PDI) override;
+	/** Executes function that generates new geometry. Returns the first new geometry index. */
+	virtual int32 ExecuteFracture(const FFractureToolContext& FractureContext) override;
+	virtual bool CanExecute() const override;
+	virtual void PostEditChangeChainProperty(struct FPropertyChangedChainEvent& PropertyChangedEvent) override;
+
+	/** Gets the UI command info for this command */
+	const TSharedPtr<FUICommandInfo>& GetUICommandInfo() const;
+
+	virtual TArray<FFractureToolContext> GetFractureToolContexts() const override;
+
+
+protected:
+	UPROPERTY(EditAnywhere, Category = Slicing)
+	TObjectPtr<UFractureConvertSettings> ConvertSettings;
+};
diff --git a/Engine/Plugins/Experimental/PlanarCutPlugin/Source/PlanarCut/Private/PlanarCut.cpp b/Engine/Plugins/Experimental/PlanarCutPlugin/Source/PlanarCut/Private/PlanarCut.cpp
--- a/Engine/Plugins/Experimental/PlanarCutPlugin/Source/PlanarCut/Private/PlanarCut.cpp
+++ b/Engine/Plugins/Experimental/PlanarCutPlugin/Source/PlanarCut/Private/PlanarCut.cpp
@@ -35,6 +35,7 @@
 #include "Engine/EngineTypes.h"
 #include "StaticMeshOperations.h"
 #include "MeshDescriptionToDynamicMesh.h"
+#include "DynamicMeshToMeshDescription.h"
 
 #include "Algo/Rotate.h"
 
@@ -151,17 +152,8 @@ namespace UE
 				Vs->GetValue(VID, V);
 			}
 
-			void ComputeTangents(FDynamicMesh3& Mesh, bool bOnlyOddMaterials, const TArrayView<const int32>& WhichMaterials, bool bRecomputeNormals = true)
+			void InitializeOverlayToPerVertexUVs(FDynamicMesh3& Mesh)
 			{
-				FDynamicMeshNormalOverlay* Normals = Mesh.Attributes()->PrimaryNormals();
-				FMeshNormals::InitializeOverlayToPerVertexNormals(Normals, !bRecomputeNormals);
-				if (bRecomputeNormals)
-				{
-					FMeshNormals::QuickRecomputeOverlayNormals(Mesh);
-				}
-
-				// Copy per-vertex UVs to a UV overlay, because that's what the tangents code uses
-				// (TODO: consider making a tangent computation path that uses vertex normals / UVs)
 				FDynamicMeshUVOverlay* UVs = Mesh.Attributes()->PrimaryUV();
 				UVs->ClearElements();
 				TArray<int> VertToUVMap;
@@ -180,6 +172,49 @@ namespace UE
 					Tri.C = VertToUVMap[Tri.C];
 					UVs->SetTriangle(TID, Tri);
 				}
+			}
+
+			void InitializeOverlayToPerVertexTangents(FDynamicMesh3& Mesh)
+			{
+				Mesh.Attributes()->EnableTangents();
+				FDynamicMeshNormalOverlay* TangentOverlays[2] = { Mesh.Attributes()->PrimaryTangents(), Mesh.Attributes()->PrimaryBiTangents() };
+				TangentOverlays[0]->ClearElements();
+				TangentOverlays[1]->ClearElements();
+				TArray<int> VertToTangentMap;
+				VertToTangentMap.SetNumUninitialized(Mesh.MaxVertexID());
+				for (int VID : Mesh.VertexIndicesItr())
+				{
+					FVector3f Tangents[2];
+					UE::PlanarCutInternals::AugmentDynamicMesh::GetTangent(Mesh, VID, Tangents[0], Tangents[1]);
+					int TID = TangentOverlays[0]->AppendElement(Tangents[0]);
+					int TID2 = TangentOverlays[1]->AppendElement(Tangents[1]);
+					check(TID == TID2);
+					VertToTangentMap[VID] = TID;
+				}
+				for (int TID : Mesh.TriangleIndicesItr())
+				{
+					FIndex3i Tri = Mesh.GetTriangle(TID);
+					Tri.A = VertToTangentMap[Tri.A];
+					Tri.B = VertToTangentMap[Tri.B];
+					Tri.C = VertToTangentMap[Tri.C];
+					TangentOverlays[0]->SetTriangle(TID, Tri);
+					TangentOverlays[1]->SetTriangle(TID, Tri);
+				}
+			}
+
+			void ComputeTangents(FDynamicMesh3& Mesh, bool bOnlyOddMaterials, const TArrayView<const int32>& WhichMaterials, bool bRecomputeNormals = true)
+			{
+				FDynamicMeshNormalOverlay* Normals = Mesh.Attributes()->PrimaryNormals();
+				FMeshNormals::InitializeOverlayToPerVertexNormals(Normals, !bRecomputeNormals);
+				if (bRecomputeNormals)
+				{
+					FMeshNormals::QuickRecomputeOverlayNormals(Mesh);
+				}
+
+				// Copy per-vertex UVs to a UV overlay, because that's what the tangents code uses
+				// (TODO: consider making a tangent computation path that uses vertex normals / UVs)
+				InitializeOverlayToPerVertexUVs(Mesh);
+				FDynamicMeshUVOverlay* UVs = Mesh.Attributes()->PrimaryUV();
 
 				FComputeTangentsOptions Options;
 				Options.bAngleWeighted = true;
@@ -2991,3 +3026,66 @@ int32 AddCollisionSampleVertices(double CollisionSampleSpacing, FGeometryCollect
 	return INDEX_NONE;
 }
 
+
+
+void ConvertToMeshDescription(
+	FMeshDescription& MeshOut,
+	FTransform& TransformOut,
+	bool bCenterPivot,
+	FGeometryCollection& Collection,
+	const TArrayView<const int32>& TransformIndices
+)
+{
+	FTransform CellsToWorld = FTransform::Identity;
+	TransformOut = FTransform::Identity;
+
+	FDynamicMeshCollection MeshCollection(&Collection, TransformIndices, CellsToWorld);
+	
+	FDynamicMesh3 CombinedMesh;
+	UE::PlanarCutInternals::AugmentDynamicMesh::Augment(CombinedMesh);
+	CombinedMesh.Attributes()->EnableTangents();
+
+	int32 NumMeshes = MeshCollection.Meshes.Num();
+	for (int32 MeshIdx = 0; MeshIdx < NumMeshes; MeshIdx++)
+	{
+		FDynamicMesh3& Mesh = MeshCollection.Meshes[MeshIdx].AugMesh;
+		const FTransform& ToCollection = MeshCollection.Meshes[MeshIdx].ToCollection;
+
+		// transform from the local geometry to the collection space
+		// unless it's just one mesh that we're going to center later anyway
+		if (!bCenterPivot || NumMeshes > 1)
+		{
+			MeshTransforms::ApplyTransform(Mesh, (FTransform3d)ToCollection);
+		}
+
+		FMeshNormals::InitializeOverlayToPerVertexNormals(Mesh.Attributes()->PrimaryNormals(), true);
+		UE::PlanarCutInternals::AugmentDynamicMesh::InitializeOverlayToPerVertexUVs(Mesh);
+		UE::PlanarCutInternals::AugmentDynamicMesh::InitializeOverlayToPerVertexTangents(Mesh);
+
+		FMergeCoincidentMeshEdges EdgeMerge(&Mesh);
+		EdgeMerge.Apply();
+		
+		if (MeshIdx > 0)
+		{
+			FDynamicMeshEditor MeshAppender(&CombinedMesh);
+			FMeshIndexMappings IndexMaps_Unused;
+			MeshAppender.AppendMesh(&Mesh, IndexMaps_Unused);
+		}
+		else
+		{
+			CombinedMesh = Mesh;
+		}
+	}
+
+	if (bCenterPivot)
+	{
+		FAxisAlignedBox3d Bounds = CombinedMesh.GetCachedBounds();
+		FVector3d Translate = -Bounds.Center();
+		MeshTransforms::Translate(CombinedMesh, Translate);
+		TransformOut = FTransform((FVector)-Translate);
+	}
+
+	FDynamicMeshToMeshDescription Converter;
+	Converter.Convert(&CombinedMesh, MeshOut, true);
+}
+
diff --git a/Engine/Plugins/Experimental/PlanarCutPlugin/Source/PlanarCut/Public/PlanarCut.h b/Engine/Plugins/Experimental/PlanarCutPlugin/Source/PlanarCut/Public/PlanarCut.h
--- a/Engine/Plugins/Experimental/PlanarCutPlugin/Source/PlanarCut/Public/PlanarCut.h
+++ b/Engine/Plugins/Experimental/PlanarCutPlugin/Source/PlanarCut/Public/PlanarCut.h
@@ -231,7 +231,7 @@ void PLANARCUT_API RecomputeNormalsAndTangents(bool bOnlyTangents, FGeometryColl
 int32 PLANARCUT_API AddCollisionSampleVertices(double TargetSpacing, FGeometryCollection& Collection, const TArrayView<const int32>& TransformIndices = TArrayView<const int32>());
 
 /**
- * Cut multiple Geometry groups inside a GeometryCollection with Planes, and add each cut cell back to the GeometryCollection as a new child of their source Geometry.  For geometries that would not be cut, nothing is added.
+ * Cut multiple Geometry groups inside a GeometryCollection with a mesh, and add each cut cell back to the GeometryCollection as a new child of their source Geometry.  For geometries that would not be cut, nothing is added.
  * 
  * @param CuttingMesh				Mesh to be used to cut the geometry collection
  * @param CuttingMeshTransform		Position of cutting mesh
@@ -253,4 +253,20 @@ int32 PLANARCUT_API CutWithMesh(
 	bool bSetDefaultInternalMaterialsFromCollection = true
 );
 
+/**
+ * Convert chosen Geometry groups inside a GeometryCollection to a single Mesh Description.
+ *
+ * @param OutputMesh				Mesh to be filled with the geometry collection geometry
+ * @param TransformOut				Transform taking output mesh geometry to local space of geometry collection
+ * @param bCenterPivot				Whether to center the geometry at the origin
+ * @param Collection				The collection to be converted
+ * @param TransformIndices			Which transform groups inside the collection to convert
+ */
+void PLANARCUT_API ConvertToMeshDescription(
+	FMeshDescription& OutputMesh,
+	FTransform& TransformOut,
+	bool bCenterPivot,
+	FGeometryCollection& Collection,
+	const TArrayView<const int32>& TransformIndices
+);
 
