Fracture Editor : Improved Fracture Editor UI
- Added context menu to the outline ( to set InitialDynamicState )
- Display InitialDynamicState in tree
- Fix the default color not playing well with selection
- Highlight parent item when selecting an item

diff --git a/Engine/Plugins/Experimental/ChaosEditor/Source/FractureEditor/FractureEditor.Build.cs b/Engine/Plugins/Experimental/ChaosEditor/Source/FractureEditor/FractureEditor.Build.cs
--- a/Engine/Plugins/Experimental/ChaosEditor/Source/FractureEditor/FractureEditor.Build.cs
+++ b/Engine/Plugins/Experimental/ChaosEditor/Source/FractureEditor/FractureEditor.Build.cs
@@ -43,6 +43,7 @@ public class FractureEditor : ModuleRules
 				"SlateCore",
 				"InputCore",
 				"EditorFramework",
+				"ToolMenus",
 				"UnrealEd",
 				"LevelEditor",
                 "GeometryCollectionEngine",
@@ -51,7 +52,8 @@ public class FractureEditor : ModuleRules
 				"MeshDescription",
 				"StaticMeshDescription",
 				"PlanarCut",
-				"Chaos"
+				"Chaos",
+
 
 				// ... add private dependencies that you statically link with here ...	
 			}
diff --git a/Engine/Plugins/Experimental/ChaosEditor/Source/FractureEditor/Private/SGeometryCollectionOutliner.cpp b/Engine/Plugins/Experimental/ChaosEditor/Source/FractureEditor/Private/SGeometryCollectionOutliner.cpp
--- a/Engine/Plugins/Experimental/ChaosEditor/Source/FractureEditor/Private/SGeometryCollectionOutliner.cpp
+++ b/Engine/Plugins/Experimental/ChaosEditor/Source/FractureEditor/Private/SGeometryCollectionOutliner.cpp
@@ -8,6 +8,7 @@
 #include "Widgets/Views/STableRow.h"
 
 #include "LevelEditor.h"
+#include "ToolMenus.h"
 #include "Modules/ModuleManager.h"
 #include "GeometryCollection/GeometryCollectionComponent.h"
 #include "FractureEditorMode.h"
@@ -17,6 +18,20 @@
 
 #define LOCTEXT_NAMESPACE "ChaosEditor"
 
+
+static FText GettextFromInitialDynamicState(int32 InitialDynamicState)
+{
+	switch (InitialDynamicState)
+	{
+	case 0: return NSLOCTEXT("Fracture", "FractureToolInitialDynamicStateNoOverride", "No Override");
+	case 1: return NSLOCTEXT("Fracture", "FractureToolInitialDynamicStateSleeping", "Sleeping");
+	case 2: return NSLOCTEXT("Fracture", "FractureToolInitialDynamicStateKinematic", "Kinematic");
+	case 3: return NSLOCTEXT("Fracture", "FractureToolInitialDynamicStateStatic", "Static");
+	default:
+		return FText();
+	}
+}
+
 void FGeometryCollectionTreeItem::OnDragLeave(const FDragDropEvent& InDragDropEvent)
 {
 	TSharedPtr<FDragDropOperation> Operation = InDragDropEvent.GetOperation();
@@ -31,6 +46,32 @@ void FGeometryCollectionTreeItem::OnDragLeave(const FDragDropEvent& InDragDropEv
 	}
 }
 
+void FGeometryCollectionTreeItem::GenerateContextMenu(UToolMenu* Menu, SGeometryCollectionOutliner& Outliner)
+{
+	auto SharedOutliner = StaticCastSharedRef<SGeometryCollectionOutliner>(Outliner.AsShared());
+
+	auto MakeDynamicStateMenu = [&Outliner](UToolMenu* Menu)
+	{
+		const FName MenuEntryNames[] =
+		{
+			"NoOverride"
+			"Sleeping",
+			"Kinematic",
+			"Static"
+		};
+		const int32 MenuEntryNamesCount = sizeof(MenuEntryNames) / sizeof(FName);
+
+		FToolMenuSection& Section = Menu->AddSection("State");
+		for (int32 Index = 0; Index < MenuEntryNamesCount; ++Index)
+		{
+			Section.AddMenuEntry(MenuEntryNames[Index], GettextFromInitialDynamicState(Index), FText(), FSlateIcon(), FUIAction(FExecuteAction::CreateRaw(&Outliner, &SGeometryCollectionOutliner::SetInitialDynamicState, Index)));
+		}
+	};
+
+	FToolMenuSection& Section = Menu->AddSection("State");
+	Section.AddSubMenu("SetInitialDynamicState", NSLOCTEXT("Fracture", "FractureToolSetInitialDynamicState", "Initial Dynamic State"), FText(),FNewToolMenuDelegate::CreateLambda(MakeDynamicStateMenu));
+}
+
 void FGeometryCollectionTreeItemBone::OnDragEnter(FDragDropEvent const& InDragDropEvent)
 {
 	TSharedPtr<FDragDropOperation> Operation = InDragDropEvent.GetOperation();
@@ -103,7 +144,6 @@ FReply FGeometryCollectionTreeItemBone::OnDrop(const FDragDropEvent& DragDropEve
 	return FReply::Unhandled();
 }
 
-
 UOutlinerSettings::UOutlinerSettings(const FObjectInitializer& ObjInit)
 	: Super(ObjInit)
 	, ItemText(EOutlinerItemNameEnum::BoneIndex)
@@ -122,7 +162,9 @@ void SGeometryCollectionOutliner::Construct(const FArguments& InArgs)
 		.OnSelectionChanged(this, &SGeometryCollectionOutliner::OnSelectionChanged)
 		.OnGenerateRow(this, &SGeometryCollectionOutliner::MakeTreeRowWidget)
 		.OnGetChildren(this, &SGeometryCollectionOutliner::OnGetChildren)
+		.OnContextMenuOpening(this, &SGeometryCollectionOutliner::OnOpenContextMenu)
 		.AllowInvisibleItemSelection(true)
+		.HighlightParentNodesForSelection(true)
 		.OnSetExpansionRecursive(this, &SGeometryCollectionOutliner::ExpandRecursive)
 	];
 }
@@ -142,6 +184,30 @@ void SGeometryCollectionOutliner::OnGetChildren(FGeometryCollectionTreeItemPtr I
 	InItem->GetChildren(OutChildren);
 }
 
+TSharedPtr<SWidget> SGeometryCollectionOutliner::OnOpenContextMenu()
+{
+	FGeometryCollectionTreeItemList SelectedItems;
+	TreeView->GetSelectedItems(SelectedItems);
+
+	if (SelectedItems.Num())
+	{
+		UToolMenus* ToolMenus = UToolMenus::Get();
+		static const FName MenuName = "SGeometryCollectionOutliner.GeometryCollectionOutlinerContextMenu";
+		if (!ToolMenus->IsMenuRegistered(MenuName))
+		{
+			ToolMenus->RegisterMenu(MenuName);
+		}
+
+		// Build up the menu for a selection
+		FToolMenuContext Context;
+		UToolMenu* Menu = ToolMenus->GenerateMenu(MenuName, Context);
+		SelectedItems[0]->GenerateContextMenu(Menu, *this);
+		return ToolMenus->GenerateWidget(Menu);
+	}
+
+	return TSharedPtr<SWidget>();
+}
+
 void SGeometryCollectionOutliner::UpdateGeometryCollection()
 {
 	TreeView->RequestTreeRefresh();
@@ -281,6 +347,29 @@ void SGeometryCollectionOutliner::OnSelectionChanged(FGeometryCollectionTreeItem
 	}
 }
 
+void SGeometryCollectionOutliner::SetInitialDynamicState(int32 InDynamicState)
+{
+	FGeometryCollectionTreeItemList SelectedItems;
+	TreeView->GetSelectedItems(SelectedItems);
+
+	for (auto& SelectedItem : SelectedItems)
+	{
+		if (UGeometryCollectionComponent* Component = SelectedItem->GetComponent())
+		{
+			if (const UGeometryCollection* RestCollection = Component->GetRestCollection())
+			{
+				if (FGeometryCollection* GeometryCollection = RestCollection->GetGeometryCollection().Get())
+				{
+					TManagedArray<int32>& InitialDynamicState = GeometryCollection->GetAttribute<int32>("InitialDynamicState", FGeometryCollection::TransformGroup);
+					TArray<int32> SelectedBones = Component->GetSelectedBones();
+					InitialDynamicState[SelectedItem->GetBoneIndex()] = InDynamicState;
+				}
+			}
+		}
+	}
+	RegenerateItems();
+}
+
 TSharedRef<ITableRow> FGeometryCollectionTreeItemComponent::MakeTreeRowWidget(const TSharedRef<STableViewBase>& InOwnerTable)
 {
 	FString ActorName = Component->GetOwner()->GetActorLabel();
@@ -476,6 +565,7 @@ TSharedRef<ITableRow> FGeometryCollectionTreeItemBone::MakeTreeRowWidget(const T
 	// Set color according to simulation type
 
 	FSlateColor TextColor(FLinearColor::Red); // default color indicates something wrong
+	int32 InitialDynamicState = INDEX_NONE;
 
 	const UGeometryCollection* RestCollection = ParentComponentItem->GetComponent()->GetRestCollection();
 	if (RestCollection && !RestCollection->IsPendingKill())
@@ -489,7 +579,7 @@ TSharedRef<ITableRow> FGeometryCollectionTreeItemBone::MakeTreeRowWidget(const T
 				break;
 
 			case FGeometryCollection::ESimulationTypes::FST_Rigid:
-				TextColor = FLinearColor::Gray;
+				TextColor = FSlateColor::UseForeground();
 				break;
 
 			case FGeometryCollection::ESimulationTypes::FST_Clustered:
@@ -500,19 +590,36 @@ TSharedRef<ITableRow> FGeometryCollectionTreeItemBone::MakeTreeRowWidget(const T
 				ensureMsgf(false, TEXT("Invalid Geometry Collection simulation type encountered."));
 				break;
 		}
+
+		InitialDynamicState = GeometryCollectionPtr->InitialDynamicState[GetBoneIndex()];
 	}
 	else
 	{
 		// Deleted rest collection
 		TextColor = FLinearColor(0.1f, 0.1f, 0.1f);
+		InitialDynamicState = INDEX_NONE;
 	}
 	
 	return SNew(STableRow<FGeometryCollectionTreeItemPtr>, InOwnerTable)
 		.Content()
-		[			
-			SNew(STextBlock)
-			.Text(ItemText)
-			.ColorAndOpacity(TextColor)
+		[
+			SNew(SHorizontalBox)
+			+ SHorizontalBox::Slot()
+				.Padding(2.0f, 4.0f)
+				.AutoWidth()
+				[
+					SNew(STextBlock)
+					.Text(ItemText)
+					.ColorAndOpacity(TextColor)
+				]
+			+ SHorizontalBox::Slot()
+				.Padding(2.0f, 2.0f)
+				.FillWidth(1.0f)
+				.HAlign(HAlign_Right)
+				[
+					SNew(STextBlock)
+					.Text(GettextFromInitialDynamicState(InitialDynamicState))
+				]
 		]
 		.OnDragDetected(this, &FGeometryCollectionTreeItem::OnDragDetected)
 		.OnDrop(this, &FGeometryCollectionTreeItem::OnDrop)
diff --git a/Engine/Plugins/Experimental/ChaosEditor/Source/FractureEditor/Private/SGeometryCollectionOutliner.h b/Engine/Plugins/Experimental/ChaosEditor/Source/FractureEditor/Private/SGeometryCollectionOutliner.h
--- a/Engine/Plugins/Experimental/ChaosEditor/Source/FractureEditor/Private/SGeometryCollectionOutliner.h
+++ b/Engine/Plugins/Experimental/ChaosEditor/Source/FractureEditor/Private/SGeometryCollectionOutliner.h
@@ -22,6 +22,7 @@
 class FGeometryCollection;
 class FGeometryCollectionTreeItem;
 class FGeometryCollectionTreeItemBone;
+class SGeometryCollectionOutliner;
 
 typedef TArray<TSharedPtr<FGeometryCollectionTreeItem>> FGeometryCollectionTreeItemList;
 typedef TSharedPtr<FGeometryCollectionTreeItem> FGeometryCollectionTreeItemPtr;
@@ -64,6 +65,9 @@ public:
 	
 	virtual void OnDragEnter(FDragDropEvent const& InDragDropEvent) {};
 	virtual void OnDragLeave(FDragDropEvent const& InDragDropEvent);
+
+	//void MakeDynamicStateMenu(UToolMenu* Menu, SGeometryCollectionOutliner& Outliner);
+	void GenerateContextMenu(UToolMenu* Menu, SGeometryCollectionOutliner& Outliner);
 };
 
 class FGeometryCollectionTreeItemComponent : public FGeometryCollectionTreeItem
@@ -158,11 +162,14 @@ public:
 
 	TSharedRef<ITableRow> MakeTreeRowWidget(FGeometryCollectionTreeItemPtr InInfo, const TSharedRef<STableViewBase>& OwnerTable);
 	void OnGetChildren(TSharedPtr<FGeometryCollectionTreeItem> InInfo, TArray< TSharedPtr<FGeometryCollectionTreeItem> >& OutChildren);
+	TSharedPtr<SWidget> OnOpenContextMenu();
 
 	void UpdateGeometryCollection();
 	void SetComponents(const TArray<UGeometryCollectionComponent*>& InNewComponents);
 	void SetBoneSelection(UGeometryCollectionComponent* RootComponent, const TArray<int32>& InSelection, bool bClearCurrentSelection);
 
+	void SetInitialDynamicState(int32 InDynamicState);
+
 	void ExpandAll();
 	void ExpandRecursive(TSharedPtr<FGeometryCollectionTreeItem> TreeItem, bool bInExpansionState) const;
 
