CIS fix for non-unity compile; include missing header for FractureToolConvert

#rnx

diff --git a/Engine/Plugins/Experimental/ChaosEditor/Source/FractureEditor/Private/FractureToolConvert.cpp b/Engine/Plugins/Experimental/ChaosEditor/Source/FractureEditor/Private/FractureToolConvert.cpp
--- a/Engine/Plugins/Experimental/ChaosEditor/Source/FractureEditor/Private/FractureToolConvert.cpp
+++ b/Engine/Plugins/Experimental/ChaosEditor/Source/FractureEditor/Private/FractureToolConvert.cpp
@@ -21,6 +21,8 @@
 #include "Engine/CollisionProfile.h"
 #include "PhysicsEngine/BodySetup.h"
 
+#include "Editor.h" // for GEditor
+
 #include "FileHelpers.h" // TODO: rm if we aren't keeping the prompt for checkout and save fn
 
 #include "PlanarCut.h"
