Fracture Mode
-- protection against invalid selections

diff --git a/Engine/Plugins/Experimental/ChaosEditor/Source/FractureEditor/Private/FractureToolContext.cpp b/Engine/Plugins/Experimental/ChaosEditor/Source/FractureEditor/Private/FractureToolContext.cpp
--- a/Engine/Plugins/Experimental/ChaosEditor/Source/FractureEditor/Private/FractureToolContext.cpp
+++ b/Engine/Plugins/Experimental/ChaosEditor/Source/FractureEditor/Private/FractureToolContext.cpp
@@ -24,6 +24,11 @@ FFractureToolContext::FFractureToolContext(UGeometryCollectionComponent* InGeome
 
 void FFractureToolContext::Sanitize()
 {
+	// Ensure that selected idices are valid
+	SelectedBones.RemoveAll([this](int32 Index) {
+		return Index == INDEX_NONE;
+		});
+	
 	// Ensure that children of a selected node are not also selected.
 	SelectedBones.RemoveAll([this](int32 Index) {
 		return !IsValidBone(Index) || HasSelectedAncestor(Index);
@@ -178,6 +183,7 @@ void FFractureToolContext::ConvertSelectionToClusterNodes()
 
 bool FFractureToolContext::HasSelectedAncestor(int32 Index) const
 {
+	if (Index == INDEX_NONE)
 	const TManagedArray<int32>& Parents = GeometryCollection->GetAttribute<int32>("Parent", FGeometryCollection::TransformGroup);
 
 	if (!ensureMsgf(Index >= 0 && Index < Parents.Num(), TEXT("Invalid index in selection: %d"), Index))
