Fracture Mode
-- fixed a garbled robomerge

diff --git a/Engine/Plugins/Experimental/ChaosEditor/Source/FractureEditor/Private/FractureToolContext.cpp b/Engine/Plugins/Experimental/ChaosEditor/Source/FractureEditor/Private/FractureToolContext.cpp
--- a/Engine/Plugins/Experimental/ChaosEditor/Source/FractureEditor/Private/FractureToolContext.cpp
+++ b/Engine/Plugins/Experimental/ChaosEditor/Source/FractureEditor/Private/FractureToolContext.cpp
@@ -183,7 +183,6 @@ void FFractureToolContext::ConvertSelectionToClusterNodes()
 
 bool FFractureToolContext::HasSelectedAncestor(int32 Index) const
 {
-	if (Index == INDEX_NONE)
 	const TManagedArray<int32>& Parents = GeometryCollection->GetAttribute<int32>("Parent", FGeometryCollection::TransformGroup);
 
 	if (!ensureMsgf(Index >= 0 && Index < Parents.Num(), TEXT("Invalid index in selection: %d"), Index))
