Add number of selected elements in the fracture editor

diff --git a/Engine/Plugins/Experimental/ChaosEditor/Source/FractureEditor/Private/FractureEditorModeToolkit.cpp b/Engine/Plugins/Experimental/ChaosEditor/Source/FractureEditor/Private/FractureEditorModeToolkit.cpp
--- a/Engine/Plugins/Experimental/ChaosEditor/Source/FractureEditor/Private/FractureEditorModeToolkit.cpp
+++ b/Engine/Plugins/Experimental/ChaosEditor/Source/FractureEditor/Private/FractureEditorModeToolkit.cpp
@@ -295,6 +295,12 @@ void FFractureEditorModeToolkit::Init(const TSharedPtr<IToolkitHost>& InitToolki
 					SAssignNew(OutlinerView, SGeometryCollectionOutliner)
 					.OnBoneSelectionChanged(this, &FFractureEditorModeToolkit::OnOutlinerBoneSelectionChanged)
 				]
+				+ SVerticalBox::Slot()
+				.AutoHeight()
+				[
+					SNew(STextBlock)
+					.Text(this, &FFractureEditorModeToolkit::GetSelectionInfo)
+				]
 			]
 		]	
 	];
@@ -1341,6 +1347,13 @@ void FFractureEditorModeToolkit::OnHistogramBoneSelectionChanged(UGeometryCollec
 	OutlinerView->SetHistogramSelection(RootComponent, SelectedBones);
 }
 
+FText FFractureEditorModeToolkit::GetSelectionInfo() const
+{
+	FString Buffer = FString::Printf(TEXT("Selected: %d"), OutlinerView->GetBoneSelectionCount());
+	return FText::FromString(Buffer);
+}
+
+
 FText FFractureEditorModeToolkit::GetStatisticsSummary() const
 {
 
diff --git a/Engine/Plugins/Experimental/ChaosEditor/Source/FractureEditor/Private/SGeometryCollectionOutliner.cpp b/Engine/Plugins/Experimental/ChaosEditor/Source/FractureEditor/Private/SGeometryCollectionOutliner.cpp
--- a/Engine/Plugins/Experimental/ChaosEditor/Source/FractureEditor/Private/SGeometryCollectionOutliner.cpp
+++ b/Engine/Plugins/Experimental/ChaosEditor/Source/FractureEditor/Private/SGeometryCollectionOutliner.cpp
@@ -274,6 +274,11 @@ void SGeometryCollectionOutliner::SetHistogramSelection(UGeometryCollectionCompo
 	}
 }
 
+int32 SGeometryCollectionOutliner::GetBoneSelectionCount() const
+{
+	return TreeView->GetSelectedItems().Num();
+}
+
 void SGeometryCollectionOutliner::SetBoneSelection(UGeometryCollectionComponent* RootComponent, const TArray<int32>& InSelection, bool bClearCurrentSelection)
 {
 	TGuardValue<bool> ExternalSelectionGuard(bPerformingSelection, true);
diff --git a/Engine/Plugins/Experimental/ChaosEditor/Source/FractureEditor/Private/SGeometryCollectionOutliner.h b/Engine/Plugins/Experimental/ChaosEditor/Source/FractureEditor/Private/SGeometryCollectionOutliner.h
--- a/Engine/Plugins/Experimental/ChaosEditor/Source/FractureEditor/Private/SGeometryCollectionOutliner.h
+++ b/Engine/Plugins/Experimental/ChaosEditor/Source/FractureEditor/Private/SGeometryCollectionOutliner.h
@@ -167,6 +167,7 @@ public:
 	void UpdateGeometryCollection();
 	void SetComponents(const TArray<UGeometryCollectionComponent*>& InNewComponents);
 	void SetBoneSelection(UGeometryCollectionComponent* RootComponent, const TArray<int32>& InSelection, bool bClearCurrentSelection);
+	int32 GetBoneSelectionCount() const;
 
 	void SetInitialDynamicState(int32 InDynamicState);
 
diff --git a/Engine/Plugins/Experimental/ChaosEditor/Source/FractureEditor/Public/FractureEditorModeToolkit.h b/Engine/Plugins/Experimental/ChaosEditor/Source/FractureEditor/Public/FractureEditorModeToolkit.h
--- a/Engine/Plugins/Experimental/ChaosEditor/Source/FractureEditor/Public/FractureEditorModeToolkit.h
+++ b/Engine/Plugins/Experimental/ChaosEditor/Source/FractureEditor/Public/FractureEditorModeToolkit.h
@@ -135,6 +135,7 @@ public:
 	int32 GetLevelCount();
 
 	FText GetStatisticsSummary() const;
+	FText GetSelectionInfo() const;
 
 	/** Returns the number of Mode specific tabs in the mode toolbar **/ 
 	const static TArray<FName> PaletteNames;
