more paranoid sanitizing on fracture tool selections

there seems to be a problem where the selection on the fracture contexts includes bones that no longer exist; I can hit this by randomly clicking the various cluster/flatten/etc buttons. The deeper problem at the root should be fixed too, but doing a bit more sanitizing should make the tool less likely to crash in the meantime.

#rnx

diff --git a/Engine/Plugins/Experimental/ChaosEditor/Source/FractureEditor/Private/FractureToolClusterMagnet.cpp b/Engine/Plugins/Experimental/ChaosEditor/Source/FractureEditor/Private/FractureToolClusterMagnet.cpp
--- a/Engine/Plugins/Experimental/ChaosEditor/Source/FractureEditor/Private/FractureToolClusterMagnet.cpp
+++ b/Engine/Plugins/Experimental/ChaosEditor/Source/FractureEditor/Private/FractureToolClusterMagnet.cpp
@@ -92,6 +92,11 @@ void UFractureToolClusterMagnet::Execute(TWeakPtr<FFractureEditorModeToolkit> In
 
 			for (TPair<int32, TArray<int32>>& Group : ClusteredSelection)
 			{
+				if (Group.Key == INDEX_NONE) // Group is top level
+				{
+					continue;
+				}
+
 				// We have the connections for the leaf nodes of our geometry collection. We want to percolate those up to the top nodes.
 				TMap<int32, TSet<int32>> TopNodeConnectivity = InitializeConnectivity(Children[Group.Key], Context.GetGeometryCollection(), Levels[Group.Key]+1);
 
diff --git a/Engine/Plugins/Experimental/ChaosEditor/Source/FractureEditor/Private/FractureToolContext.cpp b/Engine/Plugins/Experimental/ChaosEditor/Source/FractureEditor/Private/FractureToolContext.cpp
--- a/Engine/Plugins/Experimental/ChaosEditor/Source/FractureEditor/Private/FractureToolContext.cpp
+++ b/Engine/Plugins/Experimental/ChaosEditor/Source/FractureEditor/Private/FractureToolContext.cpp
@@ -24,9 +24,10 @@ FFractureToolContext::FFractureToolContext(UGeometryCollectionComponent* InGeome
 
 void FFractureToolContext::Sanitize()
 {
-	// Ensure that selected idices are valid
-	SelectedBones.RemoveAll([this](int32 Index) {
-		return Index == INDEX_NONE;
+	// Ensure that selected indices are valid
+	int NumTransforms = GeometryCollection->NumElements(FGeometryCollection::TransformGroup);
+	SelectedBones.RemoveAll([this, NumTransforms](int32 Index) {
+		return Index == INDEX_NONE || !ensure(Index < NumTransforms);
 		});
 	
 	// Ensure that children of a selected node are not also selected.
@@ -39,6 +40,12 @@ void FFractureToolContext::Sanitize()
 
 void FFractureToolContext::RemoveRootNodes()
 {
+	// Ensure that selected indices are valid
+	int NumTransforms = GeometryCollection->NumElements(FGeometryCollection::TransformGroup);
+	SelectedBones.RemoveAll([this, NumTransforms](int32 Index) {
+		return Index == INDEX_NONE || !ensure(Index < NumTransforms);
+		});
+
 	SelectedBones.RemoveAll([this](int32 Index) {
 		return FGeometryCollectionClusteringUtility::IsARootBone(this->GeometryCollection.Get(), Index);
 		});
@@ -158,6 +165,8 @@ void FFractureToolContext::ConvertSelectionToClusterNodes()
 	const TManagedArray<int32>& Parents = GeometryCollection->Parent;
 	const TManagedArray<int32>& SimulationType = GeometryCollection->SimulationType;
 
+	Sanitize();
+
 	TArray<int32> AddedClusterSelections;
 	for (int32 Index : SelectedBones)
 	{
