for non-grouped fracture, make each fracture context selection only include the bone that that context is supposed to fracture

#rnx

diff --git a/Engine/Plugins/Experimental/ChaosEditor/Source/FractureEditor/Private/FractureToolCutter.cpp b/Engine/Plugins/Experimental/ChaosEditor/Source/FractureEditor/Private/FractureToolCutter.cpp
--- a/Engine/Plugins/Experimental/ChaosEditor/Source/FractureEditor/Private/FractureToolCutter.cpp
+++ b/Engine/Plugins/Experimental/ChaosEditor/Source/FractureEditor/Private/FractureToolCutter.cpp
@@ -90,9 +90,11 @@ TArray<FFractureToolContext> UFractureToolCutterBase::GetFractureToolContexts()
 				// Generate a context for each selected node
 				for (int32 Index : FullSelection.GetSelection())
 				{
-					Contexts.Emplace(GeometryCollectionComponent);
-					FFractureToolContext& FractureContext = Contexts.Last();
+					FFractureToolContext& FractureContext = Contexts.Emplace_GetRef(GeometryCollectionComponent);
 
+					TArray<int32> Selection;
+					Selection.Add(Index);
+					FractureContext.SetSelection(Selection);
 					FractureContext.SetSeed(CutterSettings->RandomSeed > -1 ? CutterSettings->RandomSeed + Index : DefaultRandomSeed + Index);
 					FractureContext.SetBounds(BoundsToBone[Index]);
 				}
