FractureMode
-- Fixed issues with PlaceCut and Slice tools to correct diagram draw issues.

diff --git a/Engine/Plugins/Experimental/ChaosEditor/Source/FractureEditor/Private/FractureToolPlaneCut.cpp b/Engine/Plugins/Experimental/ChaosEditor/Source/FractureEditor/Private/FractureToolPlaneCut.cpp
--- a/Engine/Plugins/Experimental/ChaosEditor/Source/FractureEditor/Private/FractureToolPlaneCut.cpp
+++ b/Engine/Plugins/Experimental/ChaosEditor/Source/FractureEditor/Private/FractureToolPlaneCut.cpp
@@ -42,29 +42,31 @@ void UFractureToolPlaneCut::RegisterUICommand( FFractureEditorCommands* BindingC
 void UFractureToolPlaneCut::Render(const FSceneView* View, FViewport* Viewport, FPrimitiveDrawInterface* PDI)
 {
 	const UFracturePlaneCutSettings* LocalCutSettings = PlaneCutSettings;
-
-	if (LocalCutSettings->ReferenceActor != nullptr) // so we update with the ref actor realtime
+	if (CutterSettings->bDrawDiagram)
 	{
-		FTransform Transform(LocalCutSettings->ReferenceActor->GetActorTransform());
-		PDI->DrawPoint(Transform.GetLocation(), FLinearColor::Green, 4.f, SDPG_Foreground);
+		if (LocalCutSettings->ReferenceActor != nullptr) // so we update with the ref actor realtime
+		{
+			FTransform Transform(LocalCutSettings->ReferenceActor->GetActorTransform());
+			PDI->DrawPoint(Transform.GetLocation(), FLinearColor::Green, 4.f, SDPG_Foreground);
 
-		PDI->DrawLine(Transform.GetLocation(), Transform.GetLocation() + Transform.GetUnitAxis(EAxis::X) * 100.f, FLinearColor(255, 0, 0), SDPG_Foreground);
-		PDI->DrawLine(Transform.GetLocation(), Transform.GetLocation() + Transform.GetUnitAxis(EAxis::Y) * 100.f, FLinearColor(0, 255, 0), SDPG_Foreground);
+			PDI->DrawLine(Transform.GetLocation(), Transform.GetLocation() + Transform.GetUnitAxis(EAxis::X) * 100.f, FLinearColor(255, 0, 0), SDPG_Foreground);
+			PDI->DrawLine(Transform.GetLocation(), Transform.GetLocation() + Transform.GetUnitAxis(EAxis::Y) * 100.f, FLinearColor(0, 255, 0), SDPG_Foreground);
 
-		PDI->DrawLine(Transform.GetLocation() + Transform.GetUnitAxis(EAxis::X) * 100.f, Transform.GetLocation() + Transform.GetUnitAxis(EAxis::X) * 100.f + Transform.GetUnitAxis(EAxis::Y) * 100.f, FLinearColor(255, 0, 0), SDPG_Foreground);
-		PDI->DrawLine(Transform.GetLocation() + Transform.GetUnitAxis(EAxis::Y) * 100.f, Transform.GetLocation() + Transform.GetUnitAxis(EAxis::X) * 100.f + Transform.GetUnitAxis(EAxis::Y) * 100.f, FLinearColor(0, 255, 0), SDPG_Foreground);
-	}
-	else // draw from computed transforms
-	{
-		for (const FTransform& Transform : RenderCuttingPlanesTransforms)
+			PDI->DrawLine(Transform.GetLocation() + Transform.GetUnitAxis(EAxis::X) * 100.f, Transform.GetLocation() + Transform.GetUnitAxis(EAxis::X) * 100.f + Transform.GetUnitAxis(EAxis::Y) * 100.f, FLinearColor(255, 0, 0), SDPG_Foreground);
+			PDI->DrawLine(Transform.GetLocation() + Transform.GetUnitAxis(EAxis::Y) * 100.f, Transform.GetLocation() + Transform.GetUnitAxis(EAxis::X) * 100.f + Transform.GetUnitAxis(EAxis::Y) * 100.f, FLinearColor(0, 255, 0), SDPG_Foreground);
+		}
+		else // draw from computed transforms
 		{
-			PDI->DrawPoint(Transform.GetLocation(), FLinearColor::Green, 4.f, SDPG_Foreground);
+			for (const FTransform& Transform : RenderCuttingPlanesTransforms)
+			{
+				PDI->DrawPoint(Transform.GetLocation(), FLinearColor::Green, 4.f, SDPG_Foreground);
 
-			PDI->DrawLine(Transform.GetLocation(), Transform.GetLocation() + Transform.GetUnitAxis(EAxis::X) * RenderCuttingPlaneSize, FLinearColor(255, 0, 0), SDPG_Foreground);
-			PDI->DrawLine(Transform.GetLocation(), Transform.GetLocation() + Transform.GetUnitAxis(EAxis::Y) * RenderCuttingPlaneSize, FLinearColor(0, 255, 0), SDPG_Foreground);
+				PDI->DrawLine(Transform.GetLocation(), Transform.GetLocation() + Transform.GetUnitAxis(EAxis::X) * RenderCuttingPlaneSize, FLinearColor(255, 0, 0), SDPG_Foreground);
+				PDI->DrawLine(Transform.GetLocation(), Transform.GetLocation() + Transform.GetUnitAxis(EAxis::Y) * RenderCuttingPlaneSize, FLinearColor(0, 255, 0), SDPG_Foreground);
 
-			PDI->DrawLine(Transform.GetLocation() + Transform.GetUnitAxis(EAxis::X) * RenderCuttingPlaneSize, Transform.GetLocation() + Transform.GetUnitAxis(EAxis::X) * RenderCuttingPlaneSize + Transform.GetUnitAxis(EAxis::Y) * RenderCuttingPlaneSize, FLinearColor(255, 0, 0), SDPG_Foreground);
-			PDI->DrawLine(Transform.GetLocation() + Transform.GetUnitAxis(EAxis::Y) * RenderCuttingPlaneSize, Transform.GetLocation() + Transform.GetUnitAxis(EAxis::X) * RenderCuttingPlaneSize + Transform.GetUnitAxis(EAxis::Y) * RenderCuttingPlaneSize, FLinearColor(0, 255, 0), SDPG_Foreground);
+				PDI->DrawLine(Transform.GetLocation() + Transform.GetUnitAxis(EAxis::X) * RenderCuttingPlaneSize, Transform.GetLocation() + Transform.GetUnitAxis(EAxis::X) * RenderCuttingPlaneSize + Transform.GetUnitAxis(EAxis::Y) * RenderCuttingPlaneSize, FLinearColor(255, 0, 0), SDPG_Foreground);
+				PDI->DrawLine(Transform.GetLocation() + Transform.GetUnitAxis(EAxis::Y) * RenderCuttingPlaneSize, Transform.GetLocation() + Transform.GetUnitAxis(EAxis::X) * RenderCuttingPlaneSize + Transform.GetUnitAxis(EAxis::Y) * RenderCuttingPlaneSize, FLinearColor(0, 255, 0), SDPG_Foreground);
+			}
 		}
 	}
 }
diff --git a/Engine/Plugins/Experimental/ChaosEditor/Source/FractureEditor/Private/FractureToolSlice.cpp b/Engine/Plugins/Experimental/ChaosEditor/Source/FractureEditor/Private/FractureToolSlice.cpp
--- a/Engine/Plugins/Experimental/ChaosEditor/Source/FractureEditor/Private/FractureToolSlice.cpp
+++ b/Engine/Plugins/Experimental/ChaosEditor/Source/FractureEditor/Private/FractureToolSlice.cpp
@@ -55,11 +55,11 @@ TArray<UObject*> UFractureToolSlice::GetSettingsObjects() const
 
 void UFractureToolSlice::GenerateSliceTransforms(const FFractureToolContext& Context, TArray<FTransform>& CuttingPlaneTransforms)
 {
-	const FBox& Bounds = Context.GetBounds();
-	const FVector& Min = Context.GetBounds().Min;
-	const FVector& Max = Context.GetBounds().Max;
-	const FVector  Center = Context.GetBounds().GetCenter();
-	const FVector Extents(Context.GetBounds().Max - Context.GetBounds().Min);
+	const FBox& Bounds = Context.GetWorldBounds();
+	const FVector& Min = Bounds.Min;
+	const FVector& Max = Bounds.Max;
+	const FVector  Center = Bounds.GetCenter();
+	const FVector Extents(Max - Min);
 	const FVector HalfExtents(Extents * 0.5f);
 
 	const FVector Step(Extents.X / (SliceSettings->SlicesX+1), Extents.Y / (SliceSettings->SlicesY+1), Extents.Z / (SliceSettings->SlicesZ+1));
@@ -112,15 +112,18 @@ void UFractureToolSlice::GenerateSliceTransforms(const FFractureToolContext& Con
 
 void UFractureToolSlice::Render(const FSceneView* View, FViewport* Viewport, FPrimitiveDrawInterface* PDI)
 {
-	for (const FTransform& Transform : RenderCuttingPlanesTransforms)
+	if (CutterSettings->bDrawDiagram)
 	{
-		PDI->DrawPoint(Transform.GetLocation(), FLinearColor::Green, 4.f, SDPG_Foreground);
+		for (const FTransform& Transform : RenderCuttingPlanesTransforms)
+		{
+			PDI->DrawPoint(Transform.GetLocation(), FLinearColor::Green, 4.f, SDPG_Foreground);
 
-		PDI->DrawLine(Transform.GetLocation(), Transform.GetLocation() + Transform.GetUnitAxis(EAxis::X) * RenderCuttingPlaneSize, FLinearColor(255, 0, 0), SDPG_Foreground);
-		PDI->DrawLine(Transform.GetLocation(), Transform.GetLocation() + Transform.GetUnitAxis(EAxis::Y) * RenderCuttingPlaneSize, FLinearColor(0, 255, 0), SDPG_Foreground);
+			PDI->DrawLine(Transform.GetLocation(), Transform.GetLocation() + Transform.GetUnitAxis(EAxis::X) * RenderCuttingPlaneSize, FLinearColor(255, 0, 0), SDPG_Foreground);
+			PDI->DrawLine(Transform.GetLocation(), Transform.GetLocation() + Transform.GetUnitAxis(EAxis::Y) * RenderCuttingPlaneSize, FLinearColor(0, 255, 0), SDPG_Foreground);
 
-		PDI->DrawLine(Transform.GetLocation() + Transform.GetUnitAxis(EAxis::X) * RenderCuttingPlaneSize, Transform.GetLocation() + Transform.GetUnitAxis(EAxis::X) * RenderCuttingPlaneSize + Transform.GetUnitAxis(EAxis::Y) * RenderCuttingPlaneSize, FLinearColor(255, 0, 0), SDPG_Foreground);
-		PDI->DrawLine(Transform.GetLocation() + Transform.GetUnitAxis(EAxis::Y) * RenderCuttingPlaneSize, Transform.GetLocation() + Transform.GetUnitAxis(EAxis::X) * RenderCuttingPlaneSize + Transform.GetUnitAxis(EAxis::Y) * RenderCuttingPlaneSize, FLinearColor(0, 255, 0), SDPG_Foreground);
+			PDI->DrawLine(Transform.GetLocation() + Transform.GetUnitAxis(EAxis::X) * RenderCuttingPlaneSize, Transform.GetLocation() + Transform.GetUnitAxis(EAxis::X) * RenderCuttingPlaneSize + Transform.GetUnitAxis(EAxis::Y) * RenderCuttingPlaneSize, FLinearColor(255, 0, 0), SDPG_Foreground);
+			PDI->DrawLine(Transform.GetLocation() + Transform.GetUnitAxis(EAxis::Y) * RenderCuttingPlaneSize, Transform.GetLocation() + Transform.GetUnitAxis(EAxis::X) * RenderCuttingPlaneSize + Transform.GetUnitAxis(EAxis::Y) * RenderCuttingPlaneSize, FLinearColor(0, 255, 0), SDPG_Foreground);
+		}
 	}
 }
 
