Fracture Mode
-- Added Validate tool to cleanup superfluous geometry.

diff --git a/Engine/Plugins/Experimental/ChaosEditor/Source/FractureEditor/Private/FractureEditorCommands.h b/Engine/Plugins/Experimental/ChaosEditor/Source/FractureEditor/Private/FractureEditorCommands.h
--- a/Engine/Plugins/Experimental/ChaosEditor/Source/FractureEditor/Private/FractureEditorCommands.h
+++ b/Engine/Plugins/Experimental/ChaosEditor/Source/FractureEditor/Private/FractureEditorCommands.h
@@ -63,6 +63,7 @@ class FFractureEditorCommands : public TCommands<FFractureEditorCommands>
 		TSharedPtr< FUICommandInfo > RecomputeNormals;
 		TSharedPtr< FUICommandInfo > Resample;
 		TSharedPtr< FUICommandInfo > ConvertToMesh;
+		TSharedPtr< FUICommandInfo > Validate;
 
 		// Property Commands
 		TSharedPtr< FUICommandInfo > SetInitialDynamicState;
diff --git a/Engine/Plugins/Experimental/ChaosEditor/Source/FractureEditor/Private/FractureEditorModeToolkit.cpp b/Engine/Plugins/Experimental/ChaosEditor/Source/FractureEditor/Private/FractureEditorModeToolkit.cpp
--- a/Engine/Plugins/Experimental/ChaosEditor/Source/FractureEditor/Private/FractureEditorModeToolkit.cpp
+++ b/Engine/Plugins/Experimental/ChaosEditor/Source/FractureEditor/Private/FractureEditorModeToolkit.cpp
@@ -527,6 +527,7 @@ void FFractureEditorModeToolkit::BuildToolPalette(FName PaletteIndex, class FToo
 		ToolbarBuilder.AddToolBarButton(Commands.RecomputeNormals);
 		ToolbarBuilder.AddToolBarButton(Commands.Resample);
 		ToolbarBuilder.AddToolBarButton(Commands.ConvertToMesh);
+		ToolbarBuilder.AddToolBarButton(Commands.Validate);
 	}
 	else if (PaletteIndex == TEXT("Properties"))
 	{
diff --git a/Engine/Plugins/Experimental/ChaosEditor/Source/FractureEditor/Private/FractureToolEditing.cpp b/Engine/Plugins/Experimental/ChaosEditor/Source/FractureEditor/Private/FractureToolEditing.cpp
--- a/Engine/Plugins/Experimental/ChaosEditor/Source/FractureEditor/Private/FractureToolEditing.cpp
+++ b/Engine/Plugins/Experimental/ChaosEditor/Source/FractureEditor/Private/FractureToolEditing.cpp
@@ -97,5 +97,108 @@ void UFractureToolDeleteBranch::Execute(TWeakPtr<FFractureEditorModeToolkit> InT
 }
 
 
+
+FText UFractureToolValidate::GetDisplayText() const
+{
+	return FText(NSLOCTEXT("FractureToolEditingOps", "FractureToolValidate", "Validate"));
+}
+
+FText UFractureToolValidate::GetTooltipText() const
+{
+	return FText(NSLOCTEXT("FractureToolEditingOps", "FractureToolValidateTooltip", "Ensure that geometry collection is valid and clean."));
+}
+
+FSlateIcon UFractureToolValidate::GetToolIcon() const
+{
+	return FSlateIcon("FractureEditorStyle", "FractureEditor.Validate");
+}
+
+void UFractureToolValidate::RegisterUICommand(FFractureEditorCommands* BindingContext)
+{
+	UI_COMMAND_EXT(BindingContext, UICommandInfo, "Validate", "Validate", "Ensure that geometry collection is valid and clean.", EUserInterfaceActionType::Button, FInputChord());
+	BindingContext->Validate = UICommandInfo;
+}
+
+void UFractureToolValidate::Execute(TWeakPtr<FFractureEditorModeToolkit> InToolkit)
+{
+	if (InToolkit.IsValid())
+	{
+		FFractureEditorModeToolkit* Toolkit = InToolkit.Pin().Get();
+
+		TSet<UGeometryCollectionComponent*> GeomCompSelection;
+		GetSelectedGeometryCollectionComponents(GeomCompSelection);
+		for (UGeometryCollectionComponent* GeometryCollectionComponent : GeomCompSelection)
+		{
+			FGeometryCollectionEdit GeometryCollectionEdit = GeometryCollectionComponent->EditRestCollection();
+			if (UGeometryCollection* GeometryCollectionObject = GeometryCollectionEdit.GetRestCollection())
+			{
+				TSharedPtr<FGeometryCollection, ESPMode::ThreadSafe> GeometryCollectionPtr = GeometryCollectionObject->GetGeometryCollection();
+				if (FGeometryCollection* GeometryCollection = GeometryCollectionPtr.Get())
+				{
+					bool bDirty = false;
+
+					// Ensure that clusters do not point to geometry
+					TManagedArray<int32>& TransformToGeometry = GeometryCollection->TransformToGeometryIndex;
+					const int32 ElementCount = TransformToGeometry.Num();
+					for (int32 Idx = 0; Idx < ElementCount; ++Idx)
+					{
+						if (GeometryCollection->IsClustered(Idx) && TransformToGeometry[Idx] != INDEX_NONE)
+						{
+							TransformToGeometry[Idx] = INDEX_NONE;
+							UE_LOG(LogFractureTool, Warning, TEXT("Removed geometry index from cluster %d."), Idx);
+							bDirty = true;
+						}
+					}
+
+					// Remove any unreferenced geometry
+					TManagedArray<int32>& TransformIndex = GeometryCollection->TransformIndex;
+					const int32 GeometryCount = TransformIndex.Num();
+
+					TArray<int32> RemoveGeometry;
+					RemoveGeometry.Reserve(GeometryCount);
+
+					for (int32 Idx = 0; Idx < GeometryCount; ++Idx)
+					{
+						if ((TransformIndex[Idx] == INDEX_NONE) || (TransformToGeometry[TransformIndex[Idx]] != Idx))
+						{
+							RemoveGeometry.Add(Idx);
+							UE_LOG(LogFractureTool, Warning, TEXT("Removed dangling geometry at index %d."), Idx);
+							bDirty = true;
+						}
+					}
+
+					if (RemoveGeometry.Num() > 0)
+					{
+						GeometryCollection->RemoveElements(FGeometryCollection::GeometryGroup, RemoveGeometry);
+					}
+
+					// remove dangling clusters
+					// #todo leaving this out for the moment because we don't want to invalidate existing caches.
+					// FGeometryCollectionClusteringUtility::RemoveDanglingClusters(GeometryCollection);
+
+					if (bDirty)
+					{
+						FGeometryCollectionClusteringUtility::UpdateHierarchyLevelOfChildren(GeometryCollection, -1);
+						AddSingleRootNodeIfRequired(GeometryCollectionObject);
+						GeometryCollectionComponent->MarkRenderStateDirty();
+						GeometryCollectionObject->MarkPackageDirty();
+					}
+
+				}
+			}
+
+			GeometryCollectionComponent->InitializeEmbeddedGeometry();
+
+			FScopedColorEdit EditBoneColor = GeometryCollectionComponent->EditBoneSelection();
+			EditBoneColor.ResetBoneSelection();
+			EditBoneColor.ResetHighlightedBones();
+		}
+
+		Toolkit->OnSetLevelViewValue(-1);
+		Toolkit->SetOutlinerComponents(GeomCompSelection.Array());	
+	}
+}
+
+
 #undef LOCTEXT_NAMESPACE
 
diff --git a/Engine/Plugins/Experimental/ChaosEditor/Source/FractureEditor/Private/FractureToolEditing.h b/Engine/Plugins/Experimental/ChaosEditor/Source/FractureEditor/Private/FractureToolEditing.h
--- a/Engine/Plugins/Experimental/ChaosEditor/Source/FractureEditor/Private/FractureToolEditing.h
+++ b/Engine/Plugins/Experimental/ChaosEditor/Source/FractureEditor/Private/FractureToolEditing.h
@@ -21,3 +21,19 @@ public:
 	virtual void RegisterUICommand(FFractureEditorCommands* BindingContext) override;
 	virtual void Execute(TWeakPtr<FFractureEditorModeToolkit> InToolkit) override;
 };
+
+UCLASS(DisplayName = "Validate", Category = "FractureTools")
+class UFractureToolValidate : public UFractureActionTool
+{
+public:
+	GENERATED_BODY()
+
+	UFractureToolValidate(const FObjectInitializer& ObjInit) : Super(ObjInit) {}
+
+	// UFractureActionTool Interface
+	virtual FText GetDisplayText() const override;
+	virtual FText GetTooltipText() const override;
+	virtual FSlateIcon GetToolIcon() const override;
+	virtual void RegisterUICommand(FFractureEditorCommands* BindingContext) override;
+	virtual void Execute(TWeakPtr<FFractureEditorModeToolkit> InToolkit) override;
+};
