Fracture Mode
-- Strip tool, which removes simulation data from the geometry collection.

diff --git a/Engine/Plugins/Experimental/ChaosEditor/Source/FractureEditor/Private/FractureEditorCommands.h b/Engine/Plugins/Experimental/ChaosEditor/Source/FractureEditor/Private/FractureEditorCommands.h
--- a/Engine/Plugins/Experimental/ChaosEditor/Source/FractureEditor/Private/FractureEditorCommands.h
+++ b/Engine/Plugins/Experimental/ChaosEditor/Source/FractureEditor/Private/FractureEditorCommands.h
@@ -64,6 +64,7 @@ class FFractureEditorCommands : public TCommands<FFractureEditorCommands>
 		TSharedPtr< FUICommandInfo > Resample;
 		TSharedPtr< FUICommandInfo > ConvertToMesh;
 		TSharedPtr< FUICommandInfo > Validate;
+		TSharedPtr< FUICommandInfo > StripSimulationData;
 
 		// Property Commands
 		TSharedPtr< FUICommandInfo > SetInitialDynamicState;
diff --git a/Engine/Plugins/Experimental/ChaosEditor/Source/FractureEditor/Private/FractureEditorModeToolkit.cpp b/Engine/Plugins/Experimental/ChaosEditor/Source/FractureEditor/Private/FractureEditorModeToolkit.cpp
--- a/Engine/Plugins/Experimental/ChaosEditor/Source/FractureEditor/Private/FractureEditorModeToolkit.cpp
+++ b/Engine/Plugins/Experimental/ChaosEditor/Source/FractureEditor/Private/FractureEditorModeToolkit.cpp
@@ -528,6 +528,7 @@ void FFractureEditorModeToolkit::BuildToolPalette(FName PaletteIndex, class FToo
 		ToolbarBuilder.AddToolBarButton(Commands.Resample);
 		ToolbarBuilder.AddToolBarButton(Commands.ConvertToMesh);
 		ToolbarBuilder.AddToolBarButton(Commands.Validate);
+		ToolbarBuilder.AddToolBarButton(Commands.StripSimulationData);
 	}
 	else if (PaletteIndex == TEXT("Properties"))
 	{
diff --git a/Engine/Plugins/Experimental/ChaosEditor/Source/FractureEditor/Private/FractureToolEditing.cpp b/Engine/Plugins/Experimental/ChaosEditor/Source/FractureEditor/Private/FractureToolEditing.cpp
--- a/Engine/Plugins/Experimental/ChaosEditor/Source/FractureEditor/Private/FractureToolEditing.cpp
+++ b/Engine/Plugins/Experimental/ChaosEditor/Source/FractureEditor/Private/FractureToolEditing.cpp
@@ -3,6 +3,7 @@
 #include "FractureToolEditing.h"
 
 #include "Editor.h"
+#include "Dialogs/Dialogs.h"
 
 #include "GeometryCollection/GeometryCollection.h"
 #include "GeometryCollection/GeometryCollectionClusteringUtility.h"
@@ -201,6 +202,36 @@ void UFractureToolValidate::Execute(TWeakPtr<FFractureEditorModeToolkit> InToolk
 	}
 }
 
+bool UFractureToolValidate::StripAttributes(FGeometryCollection* GeometryCollection, const TMap<FName, TArray<FName>>& Whitelist)
+{
+	bool bChangesMade = false;
+
+	const TArray<FName> GroupNames = GeometryCollection->GroupNames();
+	for (const FName Group : GroupNames)
+	{
+		if (Whitelist.Contains(Group))
+		{
+			const TArray<FName> AttributeNames = GeometryCollection->AttributeNames(Group);
+			const TArray<FName>& WhitelistAttributes = Whitelist[Group];
+			for (const FName AttributeName : AttributeNames)
+			{
+				if (!WhitelistAttributes.Contains(AttributeName))
+				{
+					GeometryCollection->RemoveAttribute(AttributeName, Group);
+					bChangesMade = true;
+				}
+			}
+		}
+		else
+		{
+			GeometryCollection->RemoveGroup(Group);
+			bChangesMade = true;
+		}
+	}
+
+	return bChangesMade;
+}
+
 bool UFractureToolValidate::StripUnnecessaryAttributes(FGeometryCollection* GeometryCollection)
 {	
 	static TMap<FName, TArray<FName>> Necessary = {
@@ -220,7 +251,9 @@ bool UFractureToolValidate::StripUnnecessaryAttributes(FGeometryCollection* Geom
 			"Mass",
 			"ExemplarIndex",
 			"MassToLocal",
-			"DefaultMaterialIndex"
+			"DefaultMaterialIndex",
+			"Implicits",
+			"CollisionParticles"
 			} 
 		},
 		
@@ -265,33 +298,85 @@ bool UFractureToolValidate::StripUnnecessaryAttributes(FGeometryCollection* Geom
 		}
 	};
 
-	bool bChangesMade = false;
+	return StripAttributes(GeometryCollection, Necessary);
+}
 
-	const TArray<FName> GroupNames = GeometryCollection->GroupNames();
-	for (const FName Group : GroupNames)
+
+
+
+FText UFractureToolStripSimulationData::GetDisplayText() const
+{
+	return FText(NSLOCTEXT("FractureToolEditingOps", "StripSimulationData", "Strip"));
+}
+
+FText UFractureToolStripSimulationData::GetTooltipText() const
+{
+	return FText(NSLOCTEXT("FractureToolEditingOps", "FractureToolStripSimulationDataTooltip", "Remove data needed for simulation. WARNING: Geometry Collectin will no longer accurately simulate!"));
+}
+
+FSlateIcon UFractureToolStripSimulationData::GetToolIcon() const
+{
+	return FSlateIcon("FractureEditorStyle", "FractureEditor.StripSimulationData");
+}
+
+void UFractureToolStripSimulationData::RegisterUICommand(FFractureEditorCommands* BindingContext)
+{
+	UI_COMMAND_EXT(BindingContext, UICommandInfo, "StripSimulationData", "Strip", "Remove data needed for simulation. WARNING: Geometry Collection will no longer accurately simulate!", EUserInterfaceActionType::Button, FInputChord());
+	BindingContext->StripSimulationData = UICommandInfo;
+}
+
+void UFractureToolStripSimulationData::Execute(TWeakPtr<FFractureEditorModeToolkit> InToolkit)
+{
+	if (InToolkit.IsValid())
 	{
-		if (Necessary.Contains(Group))
+		FFractureEditorModeToolkit* Toolkit = InToolkit.Pin().Get();
+
+		TSet<UGeometryCollectionComponent*> GeomCompSelection;
+		GetSelectedGeometryCollectionComponents(GeomCompSelection);
+		for (UGeometryCollectionComponent* GeometryCollectionComponent : GeomCompSelection)
 		{
-			const TArray<FName> AttributeNames = GeometryCollection->AttributeNames(Group);
-			const TArray<FName>& NecessaryAttributes = Necessary[Group];
-			for (const FName AttributeName : AttributeNames)
+			FGeometryCollectionEdit GeometryCollectionEdit = GeometryCollectionComponent->EditRestCollection();
+			if (UGeometryCollection* GeometryCollectionObject = GeometryCollectionEdit.GetRestCollection())
 			{
-				if (!NecessaryAttributes.Contains(AttributeName))
+				TSharedPtr<FGeometryCollection, ESPMode::ThreadSafe> GeometryCollectionPtr = GeometryCollectionObject->GetGeometryCollection();
+				if (FGeometryCollection* GeometryCollection = GeometryCollectionPtr.Get())
 				{
-					GeometryCollection->RemoveAttribute(AttributeName, Group);
-					bChangesMade = true;
+					static TMap<FName, TArray<FName>> Removals = {
+						{ "Transform",
+							{
+							"SimulatableParticlesAttribute",
+							"Implicits",
+							"CollisionParticles"
+							}
+						}
+					};
+
+					FSuppressableWarningDialog::FSetupInfo Info(LOCTEXT("WarningStripSimulationData", "This will strip important simulation data from this GeometryCollection. This is a suitable choice if the GeometryCollection will only be used for cached playback. It will result in unpredictable behaviour if the GeometryCollection becomes dynamic. Do you want to continue?"), LOCTEXT("WarningStripSimulationData_Title", "Stripping Simulation Data"), TEXT("bStripSimulationDataWarning"), GEditorSettingsIni);
+					Info.ConfirmText = LOCTEXT("OK", "OK");
+					Info.CancelText = LOCTEXT("Cancel", "Cancel");
+					Info.bDefaultToSuppressInTheFuture = false;
+					FSuppressableWarningDialog StripSimulationDataWarning(Info);
+					if(StripSimulationDataWarning.ShowModal() != FSuppressableWarningDialog::EResult::Cancel)
+					{
+						for (const TPair<FName, TArray<FName>> Removal : Removals)
+						{
+							for (const FName AttributeName : Removal.Value)
+							{
+								if (GeometryCollection->HasAttribute(AttributeName, Removal.Key))
+								{
+									GeometryCollection->RemoveAttribute(AttributeName, Removal.Key);
+								}
+							}
+						}
+					}				
 				}
 			}
 		}
-		else
-		{
-			GeometryCollection->RemoveGroup(Group);
-			bChangesMade = true;
-		}
-	}
 
-	return bChangesMade;
+		Toolkit->SetOutlinerComponents(GeomCompSelection.Array());
+	}
 }
 
+
 #undef LOCTEXT_NAMESPACE
 
diff --git a/Engine/Plugins/Experimental/ChaosEditor/Source/FractureEditor/Private/FractureToolEditing.h b/Engine/Plugins/Experimental/ChaosEditor/Source/FractureEditor/Private/FractureToolEditing.h
--- a/Engine/Plugins/Experimental/ChaosEditor/Source/FractureEditor/Private/FractureToolEditing.h
+++ b/Engine/Plugins/Experimental/ChaosEditor/Source/FractureEditor/Private/FractureToolEditing.h
@@ -39,4 +39,21 @@ public:
 
 private:
 	bool StripUnnecessaryAttributes(FGeometryCollection* GeometryCollection);
+	bool StripAttributes(FGeometryCollection* GeometryCollection, const TMap<FName, TArray<FName>>& Whitelist);
 };
+
+UCLASS(DisplayName = "StripSimulationData", Category = "FractureTools")
+class UFractureToolStripSimulationData : public UFractureToolValidate
+{
+public:
+	GENERATED_BODY()
+
+	UFractureToolStripSimulationData(const FObjectInitializer& ObjInit) : Super(ObjInit) {}
+
+	// UFractureActionTool Interface
+	virtual FText GetDisplayText() const override;
+	virtual FText GetTooltipText() const override;
+	virtual FSlateIcon GetToolIcon() const override;
+	virtual void RegisterUICommand(FFractureEditorCommands* BindingContext) override;
+	virtual void Execute(TWeakPtr<FFractureEditorModeToolkit> InToolkit) override;
+};
\ No newline at end of file
