Fracture Mode
-- Fixed a map iteration that was happening with copies instead of references.

diff --git a/Engine/Plugins/Experimental/ChaosEditor/Source/FractureEditor/Private/FractureToolEditing.cpp b/Engine/Plugins/Experimental/ChaosEditor/Source/FractureEditor/Private/FractureToolEditing.cpp
--- a/Engine/Plugins/Experimental/ChaosEditor/Source/FractureEditor/Private/FractureToolEditing.cpp
+++ b/Engine/Plugins/Experimental/ChaosEditor/Source/FractureEditor/Private/FractureToolEditing.cpp
@@ -358,7 +358,7 @@ void UFractureToolStripSimulationData::Execute(TWeakPtr<FFractureEditorModeToolk
 					FSuppressableWarningDialog StripSimulationDataWarning(Info);
 					if(StripSimulationDataWarning.ShowModal() != FSuppressableWarningDialog::EResult::Cancel)
 					{
-						for (const TPair<FName, TArray<FName>> Removal : Removals)
+						for (const TPair<FName, TArray<FName>>& Removal : Removals)
 						{
 							for (const FName AttributeName : Removal.Value)
 							{
