Geometry Collection : Fracture Updates

- Fix for initialization crash.
- General cleanup from GUID optimizations.

diff --git a/Engine/Plugins/Experimental/ChaosEditor/Source/FractureEditor/Private/FractureEditorMode.cpp b/Engine/Plugins/Experimental/ChaosEditor/Source/FractureEditor/Private/FractureEditorMode.cpp
--- a/Engine/Plugins/Experimental/ChaosEditor/Source/FractureEditor/Private/FractureEditorMode.cpp
+++ b/Engine/Plugins/Experimental/ChaosEditor/Source/FractureEditor/Private/FractureEditorMode.cpp
@@ -19,6 +19,7 @@
 #include "GeometryCollection/GeometryCollectionObject.h"
 #include "GeometryCollection/GeometryCollectionActor.h"
 #include "GeometryCollection/GeometryCollection.h"
+#include "GeometryCollection/GeometryCollectionUtility.h"
 #include "FractureSelectionTools.h"
 #include "EditorViewportClient.h"
 #include "ScopedTransaction.h"
@@ -323,17 +324,7 @@ void FFractureEditorMode::OnActorSelectionChanged(const TArray<UObject*>& NewSel
 		{
 			// If collection does not already have guids, make them
 			FGeometryCollectionEdit RestCollectionEdit = GeometryCollectionComponent->EditRestCollection(GeometryCollection::EEditUpdate::None);
-			UGeometryCollection* RestCollection = RestCollectionEdit.GetRestCollection();
-			TSharedPtr<FGeometryCollection, ESPMode::ThreadSafe> GeometryCollection = RestCollection->GetGeometryCollection();
-			if (!GeometryCollection->HasAttribute("GUID", FGeometryCollection::TransformGroup))
-			{
-				FManagedArrayCollection::FConstructionParameters Params(FName(""), false);
-				TManagedArray<FGuid>& Guids = GeometryCollection->AddAttribute<FGuid>("GUID", FGeometryCollection::TransformGroup, Params);
-				for (int32 Idx = 0; Idx < Guids.Num(); ++Idx)
-				{
-					Guids[Idx] = FGuid::NewGuid();
-				}
-			}
+			::GeometryCollection::GenerateTemporaryGuids(RestCollectionEdit.GetRestCollection()->GetGeometryCollection().Get());
 			
 			FScopedColorEdit ShowBoneColorsEdit(GeometryCollectionComponent);
 			ShowBoneColorsEdit.SetEnableBoneSelection(true);
diff --git a/Engine/Plugins/Experimental/ChaosEditor/Source/FractureEditor/Private/FractureToolContext.cpp b/Engine/Plugins/Experimental/ChaosEditor/Source/FractureEditor/Private/FractureToolContext.cpp
--- a/Engine/Plugins/Experimental/ChaosEditor/Source/FractureEditor/Private/FractureToolContext.cpp
+++ b/Engine/Plugins/Experimental/ChaosEditor/Source/FractureEditor/Private/FractureToolContext.cpp
@@ -3,6 +3,7 @@
 #include "FractureToolContext.h"
 
 #include "GeometryCollection/GeometryCollectionClusteringUtility.h"
+#include "GeometryCollection/GeometryCollectionUtility.h"
 
 
 FFractureToolContext::FFractureToolContext(UGeometryCollectionComponent* InGeometryCollectionComponent)
@@ -53,17 +54,7 @@ void FFractureToolContext::RemoveRootNodes()
 
 void FFractureToolContext::GenerateGuids(int32 StartIdx)
 {
-	if (!GeometryCollection->HasAttribute("GUID", FGeometryCollection::TransformGroup))
-	{
-		FManagedArrayCollection::FConstructionParameters Params(FName(""), false);
-		GeometryCollection->AddAttribute<FGuid>("GUID", FGeometryCollection::TransformGroup, Params);
-	}
-
-	TManagedArray<FGuid>& Guids = GeometryCollection->GetAttribute<FGuid>("GUID", FGeometryCollection::TransformGroup);
-	for (int32 Idx = StartIdx; Idx < Guids.Num(); ++Idx)
-	{
-		Guids[Idx] = FGuid::NewGuid();
-	}
+	::GeometryCollection::GenerateTemporaryGuids(this->GeometryCollection.Get(), StartIdx, true);
 }
 
 TMap<int32, TArray<int32>> FFractureToolContext::GetClusteredSelections()
diff --git a/Engine/Plugins/Experimental/ChaosEditor/Source/FractureEditor/Private/FractureToolEditing.cpp b/Engine/Plugins/Experimental/ChaosEditor/Source/FractureEditor/Private/FractureToolEditing.cpp
--- a/Engine/Plugins/Experimental/ChaosEditor/Source/FractureEditor/Private/FractureToolEditing.cpp
+++ b/Engine/Plugins/Experimental/ChaosEditor/Source/FractureEditor/Private/FractureToolEditing.cpp
@@ -177,8 +177,6 @@ void UFractureToolValidate::Execute(TWeakPtr<FFractureEditorModeToolkit> InToolk
 					// #todo leaving this out for the moment because we don't want to invalidate existing caches.
 					// FGeometryCollectionClusteringUtility::RemoveDanglingClusters(GeometryCollection);
 
-					bDirty |= StripUnnecessaryAttributes(GeometryCollection);
-
 					if (bDirty)
 					{
 						FGeometryCollectionClusteringUtility::UpdateHierarchyLevelOfChildren(GeometryCollection, -1);
@@ -202,108 +200,6 @@ void UFractureToolValidate::Execute(TWeakPtr<FFractureEditorModeToolkit> InToolk
 	}
 }
 
-bool UFractureToolValidate::StripAttributes(FGeometryCollection* GeometryCollection, const TMap<FName, TArray<FName>>& Whitelist)
-{
-	bool bChangesMade = false;
-
-	const TArray<FName> GroupNames = GeometryCollection->GroupNames();
-	for (const FName Group : GroupNames)
-	{
-		if (Whitelist.Contains(Group))
-		{
-			const TArray<FName> AttributeNames = GeometryCollection->AttributeNames(Group);
-			const TArray<FName>& WhitelistAttributes = Whitelist[Group];
-			for (const FName AttributeName : AttributeNames)
-			{
-				if (!WhitelistAttributes.Contains(AttributeName))
-				{
-					GeometryCollection->RemoveAttribute(AttributeName, Group);
-					bChangesMade = true;
-				}
-			}
-		}
-		else
-		{
-			GeometryCollection->RemoveGroup(Group);
-			bChangesMade = true;
-		}
-	}
-
-	return bChangesMade;
-}
-
-bool UFractureToolValidate::StripUnnecessaryAttributes(FGeometryCollection* GeometryCollection)
-{	
-	static TMap<FName, TArray<FName>> Necessary = {
-		{ "Transform", 
-			{
-			"GUID",
-			"Transform",
-			"BoneColor",
-			"Parent",
-			"Children",
-			"TransformToGeometryIndex",
-			"SimulationType",
-			"StatusFlags",
-			"InitialDynamicState",
-			"SimulatableParticlesAttribute",
-			"InertiaTensor",
-			"Mass",
-			"ExemplarIndex",
-			"MassToLocal",
-			"DefaultMaterialIndex",
-			"Implicits",
-			"CollisionParticles"
-			} 
-		},
-		
-		{ "Vertices",
-			{
-			"Vertex",
-			"UV",
-			"Color",
-			"TangentU",
-			"TangentV",
-			"Normal",
-			"BoneMap"
-			}
-		},
-		
-		{ "Faces",
-			{
-			"Indices",
-			"Visible",
-			"MaterialIndex",
-			"MaterialID"
-			}
-		},
-
-		{ "Geometry",
-			{
-			"TransformIndex",
-			"BoundingBox",
-			"InnerRadius",
-			"OuterRadius",
-			"VertexStart",
-			"VertexCount",
-			"FaceStart",
-			"FaceCount"
-			}
-		},
-
-		{ "Material",
-			{
-			"Sections"
-			}
-		}
-	};
-
-	return StripAttributes(GeometryCollection, Necessary);
-}
-
-
-
-
 FText UFractureToolStripSimulationData::GetDisplayText() const
 {
 	return FText(NSLOCTEXT("FractureToolEditingOps", "StripSimulationData", "Strip"));
diff --git a/Engine/Plugins/Experimental/ChaosEditor/Source/FractureEditor/Private/FractureToolEditing.h b/Engine/Plugins/Experimental/ChaosEditor/Source/FractureEditor/Private/FractureToolEditing.h
--- a/Engine/Plugins/Experimental/ChaosEditor/Source/FractureEditor/Private/FractureToolEditing.h
+++ b/Engine/Plugins/Experimental/ChaosEditor/Source/FractureEditor/Private/FractureToolEditing.h
@@ -36,10 +36,6 @@ public:
 	virtual FSlateIcon GetToolIcon() const override;
 	virtual void RegisterUICommand(FFractureEditorCommands* BindingContext) override;
 	virtual void Execute(TWeakPtr<FFractureEditorModeToolkit> InToolkit) override;
-
-private:
-	bool StripUnnecessaryAttributes(FGeometryCollection* GeometryCollection);
-	bool StripAttributes(FGeometryCollection* GeometryCollection, const TMap<FName, TArray<FName>>& Whitelist);
 };
 
 UCLASS(DisplayName = "StripSimulationData", Category = "FractureTools")
diff --git a/Engine/Plugins/Experimental/ChaosEditor/Source/FractureEditor/Private/FractureToolGenerators.cpp b/Engine/Plugins/Experimental/ChaosEditor/Source/FractureEditor/Private/FractureToolGenerators.cpp
--- a/Engine/Plugins/Experimental/ChaosEditor/Source/FractureEditor/Private/FractureToolGenerators.cpp
+++ b/Engine/Plugins/Experimental/ChaosEditor/Source/FractureEditor/Private/FractureToolGenerators.cpp
@@ -18,6 +18,7 @@
 #include "GeometryCollection/GeometryCollection.h"
 #include "GeometryCollection/GeometryCollectionClusteringUtility.h"
 #include "GeometryCollection/GeometryCollectionConversion.h"
+#include "GeometryCollection/GeometryCollectionUtility.h"
 #include "FractureToolContext.h"
 
 
@@ -257,18 +258,7 @@ AGeometryCollectionActor* UFractureToolGenerateAsset::ConvertActorsToGeometryCol
 	}
 
 	// Add and initialize guids
-	TSharedPtr<FGeometryCollection, ESPMode::ThreadSafe> GeometryCollection = FracturedGeometryCollection->GetGeometryCollection();
-	if (!GeometryCollection->HasAttribute("GUID", FGeometryCollection::TransformGroup))
-	{
-		FManagedArrayCollection::FConstructionParameters Params(FName(""), false);
-		GeometryCollection->AddAttribute<FGuid>("GUID", FGeometryCollection::TransformGroup, Params);
-	}
-
-	TManagedArray<FGuid>& Guids = GeometryCollection->GetAttribute<FGuid>("GUID", FGeometryCollection::TransformGroup);
-	for (int32 Idx = 0; Idx < Guids.Num(); ++Idx)
-	{
-		Guids[Idx] = FGuid::NewGuid();
-	}
+	::GeometryCollection::GenerateTemporaryGuids(FracturedGeometryCollection->GetGeometryCollection().Get(), 0 , true);
 
 	return NewActor;
 }
diff --git a/Engine/Plugins/Experimental/ChaosEditor/Source/FractureEditor/Private/SGeometryCollectionHistogram.cpp b/Engine/Plugins/Experimental/ChaosEditor/Source/FractureEditor/Private/SGeometryCollectionHistogram.cpp
--- a/Engine/Plugins/Experimental/ChaosEditor/Source/FractureEditor/Private/SGeometryCollectionHistogram.cpp
+++ b/Engine/Plugins/Experimental/ChaosEditor/Source/FractureEditor/Private/SGeometryCollectionHistogram.cpp
@@ -5,6 +5,7 @@
 #include "FractureToolProperties.h"
 
 #include "GeometryCollection/GeometryCollectionComponent.h"
+#include "GeometryCollection/GeometryCollectionUtility.h"
 
 #include "Widgets/Colors/SColorBlock.h"
 #include "Widgets/Text/STextBlock.h"
@@ -139,6 +140,7 @@ FGeometryCollectionHistogramItemList FGeometryCollectionHistogramItemComponent::
 		if (Collection)
 		{
 			int32 NumElements = Collection->NumElements(FGeometryCollection::TransformGroup);
+			::GeometryCollection::GenerateTemporaryGuids(Collection);
 
 			const TManagedArray<FGuid>& Guids = Collection->GetAttribute<FGuid>("GUID", "Transform");
 			const TManagedArray<int32>& GeometryToTransform = Collection->GetAttribute<int32>("TransformIndex", "Geometry");
diff --git a/Engine/Source/Runtime/Experimental/Chaos/Private/GeometryCollection/GeometryCollectionUtility.cpp b/Engine/Source/Runtime/Experimental/Chaos/Private/GeometryCollection/GeometryCollectionUtility.cpp
--- a/Engine/Source/Runtime/Experimental/Chaos/Private/GeometryCollection/GeometryCollectionUtility.cpp
+++ b/Engine/Source/Runtime/Experimental/Chaos/Private/GeometryCollection/GeometryCollectionUtility.cpp
@@ -443,4 +443,25 @@ namespace GeometryCollection
 		}
 	}
 
+
+	void GenerateTemporaryGuids(FTransformCollection* Collection, int32 StartIdx, bool bForceInit)
+	{
+		bool bNeedsInit = false;
+		if (!Collection->HasAttribute("GUID", FTransformCollection::TransformGroup))
+		{
+			FManagedArrayCollection::FConstructionParameters Params(FName(""), false);
+			Collection->AddAttribute<FGuid>("GUID", FTransformCollection::TransformGroup, Params);
+			bNeedsInit = true;
+		}
+
+		if (bNeedsInit || bForceInit)
+		{
+			TManagedArray<FGuid>& Guids = Collection->GetAttribute<FGuid>("GUID", FTransformCollection::TransformGroup);
+			for (int32 Idx = StartIdx; Idx < Guids.Num(); ++Idx)
+			{
+				Guids[Idx] = FGuid::NewGuid();
+			}
+		}
+	}
+
 }
\ No newline at end of file
diff --git a/Engine/Source/Runtime/Experimental/Chaos/Public/GeometryCollection/GeometryCollectionUtility.h b/Engine/Source/Runtime/Experimental/Chaos/Public/GeometryCollection/GeometryCollectionUtility.h
--- a/Engine/Source/Runtime/Experimental/Chaos/Public/GeometryCollection/GeometryCollectionUtility.h
+++ b/Engine/Source/Runtime/Experimental/Chaos/Public/GeometryCollection/GeometryCollectionUtility.h
@@ -151,6 +151,17 @@ namespace GeometryCollection
 	template<class T>
 	void
 	AttributeTransfer(const FGeometryCollection * FromCollection, FGeometryCollection * ToCollection, const FName FromAttributeName, const FName ToAttributeName);
+
+	/***
+	* Generate GUID in the FTransformCollection::TransformGroup for id tracking. 
+	* The GUIDs will not be saved during serialization, and can be used to bind
+	* and entry to an location in the collection. NOTE: GUIDs are expensive to
+	* maintain so they should not be kept for cooked content. 
+	*/
+	void 
+	CHAOS_API
+	GenerateTemporaryGuids(FTransformCollection* Collection, int32 StartIdx = 0, bool bForceInit=false);
+
 };
 
 // AttributeTransfer implementation
