Fracture Tool
-- Removes Strip tool and fixed Validate tool to refresh Nanite geometry correctly.

diff --git a/Engine/Plugins/Experimental/ChaosEditor/Source/FractureEditor/Private/FractureEditorCommands.h b/Engine/Plugins/Experimental/ChaosEditor/Source/FractureEditor/Private/FractureEditorCommands.h
--- a/Engine/Plugins/Experimental/ChaosEditor/Source/FractureEditor/Private/FractureEditorCommands.h
+++ b/Engine/Plugins/Experimental/ChaosEditor/Source/FractureEditor/Private/FractureEditorCommands.h
@@ -64,7 +64,6 @@ class FFractureEditorCommands : public TCommands<FFractureEditorCommands>
 		TSharedPtr< FUICommandInfo > Resample;
 		TSharedPtr< FUICommandInfo > ConvertToMesh;
 		TSharedPtr< FUICommandInfo > Validate;
-		TSharedPtr< FUICommandInfo > StripSimulationData;
 
 		// Property Commands
 		TSharedPtr< FUICommandInfo > SetInitialDynamicState;
diff --git a/Engine/Plugins/Experimental/ChaosEditor/Source/FractureEditor/Private/FractureEditorModeToolkit.cpp b/Engine/Plugins/Experimental/ChaosEditor/Source/FractureEditor/Private/FractureEditorModeToolkit.cpp
--- a/Engine/Plugins/Experimental/ChaosEditor/Source/FractureEditor/Private/FractureEditorModeToolkit.cpp
+++ b/Engine/Plugins/Experimental/ChaosEditor/Source/FractureEditor/Private/FractureEditorModeToolkit.cpp
@@ -528,7 +528,6 @@ void FFractureEditorModeToolkit::BuildToolPalette(FName PaletteIndex, class FToo
 		ToolbarBuilder.AddToolBarButton(Commands.Resample);
 		ToolbarBuilder.AddToolBarButton(Commands.ConvertToMesh);
 		ToolbarBuilder.AddToolBarButton(Commands.Validate);
-		ToolbarBuilder.AddToolBarButton(Commands.StripSimulationData);
 	}
 	else if (PaletteIndex == TEXT("Properties"))
 	{
diff --git a/Engine/Plugins/Experimental/ChaosEditor/Source/FractureEditor/Private/FractureToolEditing.cpp b/Engine/Plugins/Experimental/ChaosEditor/Source/FractureEditor/Private/FractureToolEditing.cpp
--- a/Engine/Plugins/Experimental/ChaosEditor/Source/FractureEditor/Private/FractureToolEditing.cpp
+++ b/Engine/Plugins/Experimental/ChaosEditor/Source/FractureEditor/Private/FractureToolEditing.cpp
@@ -181,6 +181,22 @@ void UFractureToolValidate::Execute(TWeakPtr<FFractureEditorModeToolkit> InToolk
 					{
 						FGeometryCollectionClusteringUtility::UpdateHierarchyLevelOfChildren(GeometryCollection, -1);
 						AddSingleRootNodeIfRequired(GeometryCollectionObject);
+
+						// Update Nanite resource data to correctly reflect modified geometry collection data
+						{
+							GeometryCollectionObject->ReleaseResources();
+
+							if (GeometryCollectionObject->EnableNanite)
+							{
+								GeometryCollectionObject->NaniteData = UGeometryCollection::CreateNaniteData(GeometryCollection);
+							}
+							else
+							{
+								GeometryCollectionObject->NaniteData = MakeUnique<FGeometryCollectionNaniteData>();
+							}
+							GeometryCollectionObject->InitResources();
+						}
+
 						GeometryCollectionComponent->MarkRenderStateDirty();
 						GeometryCollectionObject->MarkPackageDirty();
 					}
@@ -200,79 +216,5 @@ void UFractureToolValidate::Execute(TWeakPtr<FFractureEditorModeToolkit> InToolk
 	}
 }
 
-FText UFractureToolStripSimulationData::GetDisplayText() const
-{
-	return FText(NSLOCTEXT("FractureToolEditingOps", "StripSimulationData", "Strip"));
-}
-
-FText UFractureToolStripSimulationData::GetTooltipText() const
-{
-	return FText(NSLOCTEXT("FractureToolEditingOps", "FractureToolStripSimulationDataTooltip", "Remove data needed for simulation. WARNING: Geometry Collectin will no longer accurately simulate!"));
-}
-
-FSlateIcon UFractureToolStripSimulationData::GetToolIcon() const
-{
-	return FSlateIcon("FractureEditorStyle", "FractureEditor.StripSimulationData");
-}
-
-void UFractureToolStripSimulationData::RegisterUICommand(FFractureEditorCommands* BindingContext)
-{
-	UI_COMMAND_EXT(BindingContext, UICommandInfo, "StripSimulationData", "Strip", "Remove data needed for simulation. WARNING: Geometry Collection will no longer accurately simulate!", EUserInterfaceActionType::Button, FInputChord());
-	BindingContext->StripSimulationData = UICommandInfo;
-}
-
-void UFractureToolStripSimulationData::Execute(TWeakPtr<FFractureEditorModeToolkit> InToolkit)
-{
-	if (InToolkit.IsValid())
-	{
-		FFractureEditorModeToolkit* Toolkit = InToolkit.Pin().Get();
-
-		TSet<UGeometryCollectionComponent*> GeomCompSelection;
-		GetSelectedGeometryCollectionComponents(GeomCompSelection);
-		for (UGeometryCollectionComponent* GeometryCollectionComponent : GeomCompSelection)
-		{
-			FGeometryCollectionEdit GeometryCollectionEdit = GeometryCollectionComponent->EditRestCollection();
-			if (UGeometryCollection* GeometryCollectionObject = GeometryCollectionEdit.GetRestCollection())
-			{
-				TSharedPtr<FGeometryCollection, ESPMode::ThreadSafe> GeometryCollectionPtr = GeometryCollectionObject->GetGeometryCollection();
-				if (FGeometryCollection* GeometryCollection = GeometryCollectionPtr.Get())
-				{
-					static TMap<FName, TArray<FName>> Removals = {
-						{ "Transform",
-							{
-							"SimulatableParticlesAttribute",
-							"Implicits",
-							"CollisionParticles"
-							}
-						}
-					};
-
-					FSuppressableWarningDialog::FSetupInfo Info(LOCTEXT("WarningStripSimulationData", "This will strip important simulation data from this GeometryCollection. This is a suitable choice if the GeometryCollection will only be used for cached playback. It will result in unpredictable behaviour if the GeometryCollection becomes dynamic. Do you want to continue?"), LOCTEXT("WarningStripSimulationData_Title", "Stripping Simulation Data"), TEXT("bStripSimulationDataWarning"), GEditorSettingsIni);
-					Info.ConfirmText = LOCTEXT("OK", "OK");
-					Info.CancelText = LOCTEXT("Cancel", "Cancel");
-					Info.bDefaultToSuppressInTheFuture = false;
-					FSuppressableWarningDialog StripSimulationDataWarning(Info);
-					if(StripSimulationDataWarning.ShowModal() != FSuppressableWarningDialog::EResult::Cancel)
-					{
-						for (const TPair<FName, TArray<FName>>& Removal : Removals)
-						{
-							for (const FName AttributeName : Removal.Value)
-							{
-								if (GeometryCollection->HasAttribute(AttributeName, Removal.Key))
-								{
-									GeometryCollection->RemoveAttribute(AttributeName, Removal.Key);
-								}
-							}
-						}
-					}				
-				}
-			}
-		}
-
-		Toolkit->SetOutlinerComponents(GeomCompSelection.Array());
-	}
-}
-
-
 #undef LOCTEXT_NAMESPACE
 
diff --git a/Engine/Plugins/Experimental/ChaosEditor/Source/FractureEditor/Private/FractureToolEditing.h b/Engine/Plugins/Experimental/ChaosEditor/Source/FractureEditor/Private/FractureToolEditing.h
--- a/Engine/Plugins/Experimental/ChaosEditor/Source/FractureEditor/Private/FractureToolEditing.h
+++ b/Engine/Plugins/Experimental/ChaosEditor/Source/FractureEditor/Private/FractureToolEditing.h
@@ -37,19 +37,3 @@ public:
 	virtual void RegisterUICommand(FFractureEditorCommands* BindingContext) override;
 	virtual void Execute(TWeakPtr<FFractureEditorModeToolkit> InToolkit) override;
 };
-
-UCLASS(DisplayName = "StripSimulationData", Category = "FractureTools")
-class UFractureToolStripSimulationData : public UFractureToolValidate
-{
-public:
-	GENERATED_BODY()
-
-	UFractureToolStripSimulationData(const FObjectInitializer& ObjInit) : Super(ObjInit) {}
-
-	// UFractureActionTool Interface
-	virtual FText GetDisplayText() const override;
-	virtual FText GetTooltipText() const override;
-	virtual FSlateIcon GetToolIcon() const override;
-	virtual void RegisterUICommand(FFractureEditorCommands* BindingContext) override;
-	virtual void Execute(TWeakPtr<FFractureEditorModeToolkit> InToolkit) override;
-};
\ No newline at end of file
diff --git a/Engine/Source/Runtime/Experimental/GeometryCollectionEngine/Private/GeometryCollection/GeometryCollectionSceneProxy.cpp b/Engine/Source/Runtime/Experimental/GeometryCollectionEngine/Private/GeometryCollection/GeometryCollectionSceneProxy.cpp
--- a/Engine/Source/Runtime/Experimental/GeometryCollectionEngine/Private/GeometryCollection/GeometryCollectionSceneProxy.cpp
+++ b/Engine/Source/Runtime/Experimental/GeometryCollectionEngine/Private/GeometryCollection/GeometryCollectionSceneProxy.cpp
@@ -1126,17 +1126,17 @@ FNaniteGeometryCollectionSceneProxy::FNaniteGeometryCollectionSceneProxy(UGeomet
 		MaterialSections[SectionIndex].MaterialIndex = MeshSection.MaterialID;
 	}
 
-	Nanite::FResources& Resource = GeometryCollection->NaniteData->NaniteResource;
-	const int32 NumGeometry = Resource.HierarchyRootOffsets.Num();
-	GeometryNaniteData.SetNumUninitialized(NumGeometry);
-
 	bool bHasGeometryBoundingBoxes = Collection->HasAttribute("BoundingBox", FGeometryCollection::GeometryGroup) && Collection->NumElements(FGeometryCollection::GeometryGroup);
 	bool bHasTransformBoundingBoxes = Collection->NumElements(FGeometryCollection::TransformGroup) && 
 		Collection->HasAttribute("BoundingBox", FGeometryCollection::TransformGroup) &&
 		Collection->HasAttribute("TransformToGeometryIndex", FGeometryCollection::TransformGroup);
 
+	int32 NumGeometry = 0;
 	if (bHasGeometryBoundingBoxes)
 	{
+		NumGeometry = Collection->NumElements(FGeometryCollection::GeometryGroup);
+		GeometryNaniteData.SetNumUninitialized(NumGeometry);
+
 		const TManagedArray<FBox>& BoundingBoxes = Collection->GetAttribute<FBox>("BoundingBox", FGeometryCollection::GeometryGroup);
 		for (int32 GeometryIndex = 0; GeometryIndex < NumGeometry; ++GeometryIndex)
 		{
@@ -1148,6 +1148,10 @@ FNaniteGeometryCollectionSceneProxy::FNaniteGeometryCollectionSceneProxy(UGeomet
 	}
 	else if(bHasTransformBoundingBoxes)
 	{
+		Nanite::FResources& Resource = GeometryCollection->NaniteData->NaniteResource;
+		NumGeometry = Resource.HierarchyRootOffsets.Num();
+		GeometryNaniteData.SetNumUninitialized(NumGeometry);
+		
 		const TManagedArray<FBox>& BoundingBoxes = Collection->GetAttribute<FBox>("BoundingBox", FGeometryCollection::TransformGroup);
 		const TManagedArray<int32>& TransformToGeometry = Collection->GetAttribute<int32>("TransformToGeometryIndex", FGeometryCollection::TransformGroup);
 		const int32 NumTransforms = TransformToGeometry.Num();
