Updated Fracture Icons, Promoted MakeStaticMesh to Common Library

diff --git a/Engine/Plugins/Experimental/ChaosEditor/Source/FractureEditor/Private/FractureEditorStyle.cpp b/Engine/Plugins/Experimental/ChaosEditor/Source/FractureEditor/Private/FractureEditorStyle.cpp
--- a/Engine/Plugins/Experimental/ChaosEditor/Source/FractureEditor/Private/FractureEditorStyle.cpp
+++ b/Engine/Plugins/Experimental/ChaosEditor/Source/FractureEditor/Private/FractureEditorStyle.cpp
@@ -13,7 +13,6 @@ FFractureEditorStyle::FFractureEditorStyle()
 	: FSlateStyleSet(StyleName)
 {
 	const FVector2D IconSize(20.0f, 20.0f);
-	const FVector2D SmallIconSize(20.0f, 20.0f);
 	const FVector2D LabelIconSize(16.0f, 16.0f);
 
 
@@ -21,63 +20,55 @@ FFractureEditorStyle::FFractureEditorStyle()
 	SetContentRoot(FPaths::EnginePluginsDir() / TEXT("Experimental/ChaosEditor/Content"));
 	SetCoreContentRoot(FPaths::EngineContentDir() / TEXT("Slate"));
 
-	Set("FractureEditor.Slice", new FSlateImageBrush(RootToContentDir(TEXT("FractureSlice_48x.png")), IconSize));
-	Set("FractureEditor.Slice.Small", new FSlateImageBrush(RootToContentDir(TEXT("FractureSlice_48x.png")), SmallIconSize));
-	Set("FractureEditor.Uniform", new FSlateImageBrush(RootToContentDir(TEXT("FractureUniformVoronoi_48x.png")), IconSize));
-	Set("FractureEditor.Uniform.Small", new FSlateImageBrush(RootToContentDir(TEXT("FractureUniformVoronoi_48x.png")), SmallIconSize));
-	Set("FractureEditor.Radial", new FSlateImageBrush(RootToContentDir(TEXT("FractureRadialVoronoi_48x.png")), IconSize));
-	Set("FractureEditor.Radial.Small", new FSlateImageBrush(RootToContentDir(TEXT("FractureRadialVoronoi_48x.png")), SmallIconSize));
-	Set("FractureEditor.Clustered", new FSlateImageBrush(RootToContentDir(TEXT("FractureClusteredVoronoi_48x.png")), IconSize));
-	Set("FractureEditor.Clustered.Small", new FSlateImageBrush(RootToContentDir(TEXT("FractureClusteredVoronoi_48x.png")), SmallIconSize));
-	Set("FractureEditor.Planar", new FSlateImageBrush(RootToContentDir(TEXT("FracturePlanar_48x.png")), IconSize));
-	Set("FractureEditor.Planar.Small", new FSlateImageBrush(RootToContentDir(TEXT("FracturePlanar_48x.png")), SmallIconSize));
-	Set("FractureEditor.Brick", new FSlateImageBrush(RootToContentDir(TEXT("FractureBrick_48x.png")), IconSize));
-	Set("FractureEditor.Brick.Small", new FSlateImageBrush(RootToContentDir(TEXT("FractureBrick_48x.png")), SmallIconSize));
-	Set("FractureEditor.Texture", new FSlateImageBrush(RootToContentDir(TEXT("FractureTexture_48x.png")), IconSize));
-	Set("FractureEditor.Texture.Small", new FSlateImageBrush(RootToContentDir(TEXT("FractureTexture_48x.png")), SmallIconSize));
+	Set("FractureEditor.Slice",                new IMAGE_BRUSH_SVG("FractureSliceCut", IconSize));
+	Set("FractureEditor.Uniform",              new IMAGE_BRUSH_SVG("FractureUniformVoronoi", IconSize));
+	Set("FractureEditor.Radial",               new IMAGE_BRUSH_SVG("FractureRadialVoronoi", IconSize));
+	Set("FractureEditor.Clustered",            new IMAGE_BRUSH_SVG("FractureClusterVoronoi", IconSize));
+	Set("FractureEditor.Planar",               new IMAGE_BRUSH_SVG("FracturePlanarCut", IconSize));
+	Set("FractureEditor.Brick",                new IMAGE_BRUSH_SVG("FractureBrick", IconSize));
+	Set("FractureEditor.Texture",              new IMAGE_BRUSH_SVG("FractureTexture", IconSize));
+	Set("FractureEditor.Mesh",                 new IMAGE_BRUSH_SVG("FractureTexture", IconSize));
 
 	// This is a bit of magic.  When you pass a command your Builder.AddToolBarButton, it will automatically try to find 
 	// and Icon with the same name as the command and TCommand<> Context Name.  
-	// format <Context>.<CommandName>[.Small]
-	Set("FractureEditor.SelectAll", new FSlateImageBrush(RootToContentDir(TEXT("SelectAll_48x.png")), IconSize));
-	Set("FractureEditor.SelectAll.Small", new FSlateImageBrush(RootToContentDir(TEXT("SelectAll_48x.png")), SmallIconSize));
-	Set("FractureEditor.SelectNone", new FSlateImageBrush(RootToContentDir(TEXT("DeselectAll_48x.png")), IconSize));
-	Set("FractureEditor.SelectNone.Small", new FSlateImageBrush(RootToContentDir(TEXT("DeselectAll_48x.png")), SmallIconSize));
-	Set("FractureEditor.SelectNeighbors", new FSlateImageBrush(RootToContentDir(TEXT("SelectNeighbor_48x.png")), IconSize));
-	Set("FractureEditor.SelectNeighbors.Small", new FSlateImageBrush(RootToContentDir(TEXT("SelectNeighbor_48x.png")), SmallIconSize));
-	Set("FractureEditor.SelectSiblings", new FSlateImageBrush(RootToContentDir(TEXT("SelectSiblings_48x.png")), IconSize));
-	Set("FractureEditor.SelectSiblings.Small", new FSlateImageBrush(RootToContentDir(TEXT("SelectSiblings_48x.png")), SmallIconSize));
-	Set("FractureEditor.SelectAllInCluster", new FSlateImageBrush(RootToContentDir(TEXT("SelectAllInCluster_48x.png")), IconSize));
-	Set("FractureEditor.SelectAllInCluster.Small", new FSlateImageBrush(RootToContentDir(TEXT("SelectAllInCluster_48x.png")), SmallIconSize));
-	Set("FractureEditor.SelectInvert", new FSlateImageBrush(RootToContentDir(TEXT("SelectInvert_48x.png")), IconSize));
-	Set("FractureEditor.SelectInvert.Small", new FSlateImageBrush(RootToContentDir(TEXT("SelectInvert_48x.png")), SmallIconSize));
-
-	Set("FractureEditor.AutoCluster", new FSlateImageBrush(RootToContentDir(TEXT("AutoCluster_48x.png")), IconSize));
-	Set("FractureEditor.AutoCluster.Small", new FSlateImageBrush(RootToContentDir(TEXT("AutoCluster_48x.png")), SmallIconSize));
-	Set("FractureEditor.Cluster", new FSlateImageBrush(RootToContentDir(TEXT("Cluster_48x.png")), IconSize));
-	Set("FractureEditor.Cluster.Small", new FSlateImageBrush(RootToContentDir(TEXT("Cluster_48x.png")), SmallIconSize));
-	Set("FractureEditor.Uncluster", new FSlateImageBrush(RootToContentDir(TEXT("Uncluster_48x.png")), IconSize));
-	Set("FractureEditor.Uncluster.Small", new FSlateImageBrush(RootToContentDir(TEXT("Uncluster_48x.png")), SmallIconSize));
-	Set("FractureEditor.FlattenToLevel", new FSlateImageBrush(RootToContentDir(TEXT("FlattenToLevel_48x.png")), IconSize));
-	Set("FractureEditor.FlattenToLevel.Small", new FSlateImageBrush(RootToContentDir(TEXT("FlattenToLevel_48x.png")), SmallIconSize));
-	Set("FractureEditor.Flatten", new FSlateImageBrush(RootToContentDir(TEXT("Flatten_48x.png")), IconSize));
-	Set("FractureEditor.Flatten.Small", new FSlateImageBrush(RootToContentDir(TEXT("Flatten_48x.png")), SmallIconSize));
-	Set("FractureEditor.Merge", new FSlateImageBrush(RootToContentDir(TEXT("Merge_48x.png")), IconSize));
-	Set("FractureEditor.Merge.Small", new FSlateImageBrush(RootToContentDir(TEXT("Merge_48x.png")), SmallIconSize));
-	Set("FractureEditor.MoveUp", new FSlateImageBrush(RootToContentDir(TEXT("MoveUp_48x.png")), IconSize));
-	Set("FractureEditor.MoveUp.Small", new FSlateImageBrush(RootToContentDir(TEXT("MoveUp_48x.png")), SmallIconSize));
+	Set("FractureEditor.SelectAll",            new IMAGE_BRUSH_SVG("FractureSelectAll", IconSize));
+	Set("FractureEditor.SelectNone",           new IMAGE_BRUSH_SVG("FractureDeselectAll", IconSize));
+	Set("FractureEditor.SelectNeighbors",      new IMAGE_BRUSH_SVG("FractureSelectNeighbor", IconSize));
+	Set("FractureEditor.SelectSiblings",       new IMAGE_BRUSH_SVG("FractureSelectSiblings", IconSize));
+	Set("FractureEditor.SelectAllInCluster",   new IMAGE_BRUSH_SVG("FractureSelectAllInCluster", IconSize));
+	Set("FractureEditor.SelectInvert",         new IMAGE_BRUSH_SVG("FractureSelectInvert", IconSize));
 
+	Set("FractureEditor.AutoCluster",          new IMAGE_BRUSH_SVG("FractureAutoCluster", IconSize));
+	Set("FractureEditor.ClusterMagnet",        new IMAGE_BRUSH_SVG("FractureMagnet", IconSize));
+	Set("FractureEditor.Cluster",              new IMAGE_BRUSH_SVG("FractureCluster", IconSize));
+	Set("FractureEditor.Uncluster",            new IMAGE_BRUSH_SVG("FractureUncluster", IconSize));
+	Set("FractureEditor.FlattenToLevel",       new IMAGE_BRUSH_SVG("FractureFlattenToLevel", IconSize));
+	Set("FractureEditor.Flatten",              new IMAGE_BRUSH_SVG("FractureFlatten", IconSize));
+	Set("FractureEditor.Merge",                new IMAGE_BRUSH_SVG("FractureMerge", IconSize));
+	Set("FractureEditor.MoveUp",               new IMAGE_BRUSH_SVG("FractureMoveUpRow", IconSize));
+
+	Set("FractureEditor.AddEmbeddedGeometry",  new IMAGE_BRUSH_SVG("FractureEmbed", IconSize));
+	Set("FractureEditor.AutoEmbedGeometry",    new IMAGE_BRUSH_SVG("FractureAutoEmbed", IconSize));
+
+	Set("FractureEditor.AutoUV",               new IMAGE_BRUSH_SVG("FractureAutoUV", IconSize));
+
+	Set("FractureEditor.DeleteBranch",         new IMAGE_BRUSH_SVG("FracturePrune", IconSize));
+
+	Set("FractureEditor.RecomputeNormals",     new IMAGE_BRUSH_SVG("FractureNormals", IconSize));
+	Set("FractureEditor.Resample",             new IMAGE_BRUSH_SVG("FractureResample", IconSize));
+
+	Set("FractureEditor.ToMesh",        	   new CORE_IMAGE_BRUSH_SVG("../Editor/Slate/Starship/Common/MakeStaticMesh", IconSize));
+	Set("FractureEditor.Validate",        	   new CORE_IMAGE_BRUSH_SVG("../Editor/Slate/Starship/Common/Test", IconSize));
 
 	// View Settings
-	Set("FractureEditor.Exploded", new FSlateImageBrush(RootToContentDir(TEXT("MiniExploded_32x.png")), LabelIconSize));
-	Set("FractureEditor.Levels", new FSlateImageBrush(RootToContentDir(TEXT("MiniLevel_32x.png")), LabelIconSize));
-	Set("FractureEditor.Visibility", new FSlateImageBrush(RootToContentDir(TEXT("GeneralVisibility_48x.png")), SmallIconSize));
-	Set("FractureEditor.ToggleShowBoneColors", new FSlateImageBrush(RootToContentDir(TEXT("GeneralVisibility_48x.png")), IconSize));
-	Set("FractureEditor.ToggleShowBoneColors.Small", new FSlateImageBrush(RootToContentDir(TEXT("GeneralVisibility_48x.png")), SmallIconSize));
-	Set("FractureEditor.ViewUpOneLevel", new FSlateImageBrush(RootToContentDir(TEXT("LevelViewUp_48x.png")), IconSize));
-	Set("FractureEditor.ViewUpOneLevel.Small", new FSlateImageBrush(RootToContentDir(TEXT("LevelViewUp_48x.png")), SmallIconSize));
-	Set("FractureEditor.ViewDownOneLevel", new FSlateImageBrush(RootToContentDir(TEXT("LevelViewDown_48x.png")), IconSize));
-	Set("FractureEditor.ViewDownOneLevel.Small", new FSlateImageBrush(RootToContentDir(TEXT("LevelViewDown_48x.png")), SmallIconSize));
+	Set("FractureEditor.Exploded",             new IMAGE_BRUSH_SVG("FractureMiniExploded", LabelIconSize));
+	Set("FractureEditor.Levels",               new IMAGE_BRUSH_SVG("FractureMiniLevel", LabelIconSize));
+	Set("FractureEditor.Visibility",           new IMAGE_BRUSH_SVG("FractureGeneralVisibility", IconSize));
+	Set("FractureEditor.ToggleShowBoneColors", new IMAGE_BRUSH_SVG("FractureGeneralVisibility", IconSize));
+	Set("FractureEditor.ViewUpOneLevel",       new IMAGE_BRUSH_SVG("FractureLevelViewUp", IconSize));
+	Set("FractureEditor.ViewDownOneLevel",     new IMAGE_BRUSH_SVG("FractureLevelViewDown", IconSize));
+
+	Set("FractureEditor.SetInitialDynamicState", new IMAGE_BRUSH_SVG("FractureState", IconSize));
 
 	Set("FractureEditor.SpinBox", FSpinBoxStyle(FEditorStyle::GetWidgetStyle<FSpinBoxStyle>("SpinBox"))
 		.SetTextPadding(FMargin(0))
@@ -89,20 +80,17 @@ FFractureEditorStyle::FFractureEditorStyle()
 		.SetArrowsImage(FSlateNoResource())
 	);
 
-	// Set("FractureEditor.Exploded.Small", new FSlateImageBrush(RootToContentDir(TEXT("Exploded_32x.png")), SmallIconSize));
 
-	Set("FractureEditor.GenerateAsset", new FSlateImageBrush(RootToContentDir(TEXT("GenerateAsset_48x.png")), IconSize));
-	Set("FractureEditor.GenerateAsset.Small", new FSlateImageBrush(RootToContentDir(TEXT("GenerateAsset_48x.png")), SmallIconSize));
-	Set("FractureEditor.ResetAsset", new IMAGE_BRUSH_SVG("Starship/FractureReset", FVector2D(20.0f, 20.0f)));
+	Set("FractureEditor.GenerateAsset", new IMAGE_BRUSH_SVG("FractureGenerateAsset", IconSize));
+	Set("FractureEditor.ResetAsset", new IMAGE_BRUSH_SVG("FractureReset", FVector2D(20.0f, 20.0f)));
 
 	if (FCoreStyle::IsStarshipStyle())
 	{
-		Set("LevelEditor.FractureMode", new IMAGE_BRUSH_SVG("Starship/fracture", FVector2D(20.0f, 20.0f)));
+		Set("LevelEditor.FractureMode", new IMAGE_BRUSH_SVG("fracture", FVector2D(20.0f, 20.0f)));
 	}
 	else
 	{
-		Set("LevelEditor.FractureMode", new FSlateImageBrush(RootToContentDir(TEXT("FractureMode.png")), FVector2D(40.0f, 40.0f)));
-		Set("LevelEditor.FractureMode.Small", new FSlateImageBrush(RootToContentDir(TEXT("FractureMode.png")), FVector2D(20.0f, 20.0f)));
+		Set("LevelEditor.FractureMode", new IMAGE_BRUSH("FractureMode.png", FVector2D(40.0f, 40.0f)));
 	}
 
 	FSlateStyleRegistry::RegisterSlateStyle(*this);
diff --git a/Engine/Plugins/Experimental/ChaosEditor/Source/FractureEditor/Private/FractureToolConvert.cpp b/Engine/Plugins/Experimental/ChaosEditor/Source/FractureEditor/Private/FractureToolConvert.cpp
--- a/Engine/Plugins/Experimental/ChaosEditor/Source/FractureEditor/Private/FractureToolConvert.cpp
+++ b/Engine/Plugins/Experimental/ChaosEditor/Source/FractureEditor/Private/FractureToolConvert.cpp
@@ -61,7 +61,7 @@ FText UFractureToolConvert::GetTooltipText() const
 
 FSlateIcon UFractureToolConvert::GetToolIcon() const
 {
-	return FSlateIcon("FractureEditorStyle", "FractureEditor.Convert");
+	return FSlateIcon("FractureEditorStyle", "FractureEditor.ToMesh");
 }
 
 void UFractureToolConvert::RegisterUICommand(FFractureEditorCommands* BindingContext)
diff --git a/Engine/Source/Editor/EditorStyle/Private/StarshipStyle.cpp b/Engine/Source/Editor/EditorStyle/Private/StarshipStyle.cpp
--- a/Engine/Source/Editor/EditorStyle/Private/StarshipStyle.cpp
+++ b/Engine/Source/Editor/EditorStyle/Private/StarshipStyle.cpp
@@ -425,6 +425,7 @@ void FStarshipEditorStyle::FStyle::SetupGeneralStyles()
 		Set( "Icons.Contact", new IMAGE_BRUSH( "Icons/icon_mail_16x", Icon16x16 ) );
 
 		Set( "Icons.DirtyBadge", new IMAGE_BRUSH_SVG("Starship/Common/DirtyBadge", Icon12x12));
+		Set( "Icons.MakeStaticMesh", new IMAGE_BRUSH_SVG("Starship/Common/MakeStaticMesh", Icon16x16));
 	}
 
 	Set( "WarningStripe", new IMAGE_BRUSH( "Common/WarningStripe", FVector2D(20,6), FLinearColor::White, ESlateBrushTileType::Horizontal ) );
@@ -5000,7 +5001,7 @@ void FStarshipEditorStyle::FStyle::SetupPersonaStyle()
 		// persona extras
 		Set("Persona.ConvertAnimationGraph", new IMAGE_BRUSH("Old/Graph/ConvertIcon", Icon40x40));
 		Set("Persona.ReimportAsset", new CORE_IMAGE_BRUSH_SVG("Starship/Common/import", Icon20x20));
-		Set("Persona.ConvertToStaticMesh", new IMAGE_BRUSH_SVG("Starship/Persona/MakeStaticMesh", Icon20x20));
+		Set("Persona.ConvertToStaticMesh", new IMAGE_BRUSH_SVG("Starship/Common/MakeStaticMesh", Icon20x20));
 		Set("Persona.BakeMaterials", new IMAGE_BRUSH("Icons/icon_tab_Layers_40x", Icon40x40));
 
 		// Anim Slot Manager
