Chaos : Fracture Mode
- Added call to temporary GUID generation.

diff --git a/Engine/Plugins/Experimental/ChaosEditor/Source/FractureEditor/Private/SGeometryCollectionOutliner.cpp b/Engine/Plugins/Experimental/ChaosEditor/Source/FractureEditor/Private/SGeometryCollectionOutliner.cpp
--- a/Engine/Plugins/Experimental/ChaosEditor/Source/FractureEditor/Private/SGeometryCollectionOutliner.cpp
+++ b/Engine/Plugins/Experimental/ChaosEditor/Source/FractureEditor/Private/SGeometryCollectionOutliner.cpp
@@ -486,6 +486,7 @@ void FGeometryCollectionTreeItemComponent::RegenerateChildren()
 		{
 			int32 NumElements = Collection->NumElements(FGeometryCollection::TransformGroup);
 
+			GeometryCollection::GenerateTemporaryGuids(Collection);
 			const TManagedArray<FGuid>& Guids = Collection->GetAttribute<FGuid>("GUID", "Transform");
 			const TManagedArray<int32>& Parents = Collection->Parent;
 			// const TManagedArray<FString>& BoneNames = CollectionPtr->BoneName;
