Chaos : CIS FIx

diff --git a/Engine/Plugins/Experimental/ChaosEditor/Source/FractureEditor/Private/SGeometryCollectionOutliner.cpp b/Engine/Plugins/Experimental/ChaosEditor/Source/FractureEditor/Private/SGeometryCollectionOutliner.cpp
--- a/Engine/Plugins/Experimental/ChaosEditor/Source/FractureEditor/Private/SGeometryCollectionOutliner.cpp
+++ b/Engine/Plugins/Experimental/ChaosEditor/Source/FractureEditor/Private/SGeometryCollectionOutliner.cpp
@@ -15,6 +15,7 @@
 #include "ScopedTransaction.h"
 #include "FractureSettings.h"
 #include "GeometryCollectionOutlinerDragDrop.h"
+#include "GeometryCollection/GeometryCollectionUtility.h"
 
 #define LOCTEXT_NAMESPACE "ChaosEditor"
 
