Add spatial gradients + option to not re-atlas to fracture AutoUV

diff --git a/Engine/Plugins/Experimental/ChaosEditor/Source/FractureEditor/Private/FractureToolUV.cpp b/Engine/Plugins/Experimental/ChaosEditor/Source/FractureEditor/Private/FractureToolUV.cpp
--- a/Engine/Plugins/Experimental/ChaosEditor/Source/FractureEditor/Private/FractureToolUV.cpp
+++ b/Engine/Plugins/Experimental/ChaosEditor/Source/FractureEditor/Private/FractureToolUV.cpp
@@ -193,20 +193,23 @@ bool UFractureToolAutoUV::SaveGeneratedTexture(UTexture2D* GeneratedTexture, FSt
 
 int32 UFractureToolAutoUV::ExecuteFracture(const FFractureToolContext& FractureContext)
 {
-	if (FractureContext.IsValid())
+	if (FractureContext.GetGeometryCollection().IsValid())
 	{
-		FScopedSlowTask UVTask(3, LOCTEXT("StartingAutoUV", "Automatically laying out and texturing internal surfaces"));
+		FScopedSlowTask UVTask(AutoUVSettings->bDoUVLayout ? 3 : 2, LOCTEXT("StartingAutoUV", "Automatically laying out and texturing internal surfaces"));
 		UVTask.MakeDialog();
 
 		FGeometryCollection& Collection = *FractureContext.GetGeometryCollection();
 		
 		int32 OutputRes = (int32)AutoUVSettings->Resolution;
 
-		UVTask.EnterProgressFrame(1, LOCTEXT("LayOutUVIslands", "Laying out UV islands"));
-		if (!UE::PlanarCut::UVLayout(Collection, OutputRes, AutoUVSettings->GutterSize))
+		if (AutoUVSettings->bDoUVLayout)
 		{
-			// failed to do layout
-			return INDEX_NONE;
+			UVTask.EnterProgressFrame(1, LOCTEXT("LayOutUVIslands", "Laying out UV islands"));
+			if (!UE::PlanarCut::UVLayout(Collection, OutputRes, AutoUVSettings->GutterSize))
+			{
+				// failed to do layout
+				return INDEX_NONE;
+			}
 		}
 
 		UVTask.EnterProgressFrame(1, LOCTEXT("TexturingSurfaces", "Texturing internal surfaces"));
@@ -220,14 +223,27 @@ int32 UFractureToolAutoUV::ExecuteFracture(const FFractureToolContext& FractureC
 		ImageBuilder.Clear(FVector4f(0, 0, 0, 0));
 
 		typedef UE::PlanarCut::EBakeAttributes EBakeAttributes;
+		FIndex4i Attributes;
 		// Note: Ordering of these attributes should match the order and comments in the AutoUVSettings struct
 		//		 Update the order and comments there if you change the ordering here.
-		FIndex4i Attributes(
-			AutoUVSettings->bDistToOuter ? (int32)EBakeAttributes::DistanceToExternal : 0, 
-			AutoUVSettings->bAmbientOcclusion ? (int32)EBakeAttributes::AmbientOcclusion : 0, 
-			AutoUVSettings->bSmoothedCurvature ? (int32)EBakeAttributes::Curvature : 0,
-			AutoUVSettings->bZNormal ? (int32)EBakeAttributes::NormalZ : 0
-		);
+		if (AutoUVSettings->BakeTextureType == ETextureType::ThicknessAndSurfaceAttributes)
+		{
+			Attributes = FIndex4i(
+				AutoUVSettings->bDistToOuter ? (int32)EBakeAttributes::DistanceToExternal : 0,
+				AutoUVSettings->bAmbientOcclusion ? (int32)EBakeAttributes::AmbientOcclusion : 0,
+				AutoUVSettings->bSmoothedCurvature ? (int32)EBakeAttributes::Curvature : 0,
+				AutoUVSettings->bZNormal ? (int32)EBakeAttributes::NormalZ : 0
+			);
+		}
+		else
+		{
+			Attributes = FIndex4i(
+				(int32)EBakeAttributes::PositionX,
+				(int32)EBakeAttributes::PositionY,
+				(int32)EBakeAttributes::PositionZ,
+				0
+			);
+		}
 		UE::PlanarCut::FTextureAttributeSettings AttribSettings;
 		AttribSettings.ToExternal_MaxDistance = AutoUVSettings->MaxDistance;
 		AttribSettings.AO_Rays = AutoUVSettings->OcclusionRays;
diff --git a/Engine/Plugins/Experimental/ChaosEditor/Source/FractureEditor/Private/FractureToolUV.h b/Engine/Plugins/Experimental/ChaosEditor/Source/FractureEditor/Private/FractureToolUV.h
--- a/Engine/Plugins/Experimental/ChaosEditor/Source/FractureEditor/Private/FractureToolUV.h
+++ b/Engine/Plugins/Experimental/ChaosEditor/Source/FractureEditor/Private/FractureToolUV.h
@@ -23,6 +23,12 @@ enum class EAutoUVTextureResolution : int32
 	Resolution8192 = 8192 UMETA(DisplayName = "8192 x 8192")
 };
 
+UENUM()
+enum class ETextureType
+{
+	ThicknessAndSurfaceAttributes,
+	SpatialGradients
+};
 
 /** Settings specifically related to the one-time destructive fracturing of a mesh **/
 UCLASS(config = EditorPerProjectUserSettings)
@@ -36,6 +42,10 @@ public:
 		: Super(ObjInit)
 	{}
 
+	/** Whether to layout UVs (to ensure unique texels for each face), or just use the existing UVs */
+	UPROPERTY(EditAnywhere, Category = Atlas)
+	bool bDoUVLayout = true;
+
 	/** The pixel resolution of the generated map */
 	UPROPERTY(EditAnywhere, Category = MapSettings)
 	EAutoUVTextureResolution Resolution = EAutoUVTextureResolution::Resolution512;
@@ -52,59 +62,63 @@ public:
 	UPROPERTY(EditAnywhere, Category = MapSettings)
 	bool bPromptToSave = true;
 
-	/** Bake the distance to the external surface to a texture channel (red) */
+	/** Which standard set of texture channels to bake */
 	UPROPERTY(EditAnywhere, Category = AttributesToBake)
+	ETextureType BakeTextureType = ETextureType::ThicknessAndSurfaceAttributes;
+
+	/** Bake the distance to the external surface to a texture channel (red) */
+	UPROPERTY(EditAnywhere, Category = AttributesToBake, meta = (EditCondition = "BakeTextureType == ETextureType::ThicknessAndSurfaceAttributes", EditConditionHides))
 	bool bDistToOuter = true;
 
 	/** Bake the ambient occlusion of each bone (considered separately) to a texture channel (green) */
-	UPROPERTY(EditAnywhere, Category = AttributesToBake)
+	UPROPERTY(EditAnywhere, Category = AttributesToBake, meta = (EditCondition = "BakeTextureType == ETextureType::ThicknessAndSurfaceAttributes", EditConditionHides))
 	bool bAmbientOcclusion = true;
 
 	/**
 	 * Bake a smoothed curvature metric to a texture channel (blue)
 	 * Specifically, this is the mean curvature of a smoothed copy of each fractured piece, baked back to the respective fracture piece.
 	 */
-	UPROPERTY(EditAnywhere, Category = AttributesToBake)
+	UPROPERTY(EditAnywhere, Category = AttributesToBake, meta = (EditCondition = "BakeTextureType == ETextureType::ThicknessAndSurfaceAttributes", EditConditionHides))
 	bool bSmoothedCurvature = true;
 
 	/** Bake the Z-component of the normal to a texture channel (alpha) */
-	UPROPERTY(EditAnywhere, Category = AttributesToBake)
+	UPROPERTY(EditAnywhere, Category = AttributesToBake, meta = (EditCondition = "BakeTextureType == ETextureType::ThicknessAndSurfaceAttributes", EditConditionHides))
 	bool bZNormal = true;
 
 	/** Max distance to search for the outer mesh surface */
-	UPROPERTY(EditAnywhere, Category = DistToOuterSettings, meta = (EditCondition = "bDistToOuter", EditConditionHides, UIMin = "1", UIMax = "100", ClampMin = ".01", ClampMax = "1000"))
+	UPROPERTY(EditAnywhere, Category = DistToOuterSettings, meta = (EditCondition = "BakeTextureType == ETextureType::ThicknessAndSurfaceAttributes && bDistToOuter", EditConditionHides, UIMin = "1", UIMax = "100", ClampMin = ".01", ClampMax = "1000"))
 	double MaxDistance = 100;
 
 	/** Number of occlusion rays */
-	UPROPERTY(EditAnywhere, Category = AmbientOcclusionSettings, meta = (EditCondition = "bAmbientOcclusion", EditConditionHides, UIMin = "1", UIMax = "1024", ClampMin = "0", ClampMax = "50000"))
+	UPROPERTY(EditAnywhere, Category = AmbientOcclusionSettings, meta = (EditCondition = "BakeTextureType == ETextureType::ThicknessAndSurfaceAttributes && bAmbientOcclusion", EditConditionHides, UIMin = "1", UIMax = "1024", ClampMin = "0", ClampMax = "50000"))
 	int OcclusionRays = 16;
 
 	/** Pixel Radius of Gaussian Blur Kernel applied to AO map (0 will apply no blur) */
-	UPROPERTY(EditAnywhere, Category = AmbientOcclusionSettings, meta = (EditCondition = "bAmbientOcclusion", EditConditionHides, UIMin = "0", UIMax = "10.0", ClampMin = "0", ClampMax = "100.0"))
+	UPROPERTY(EditAnywhere, Category = AmbientOcclusionSettings, meta = (EditCondition = "BakeTextureType == ETextureType::ThicknessAndSurfaceAttributes && bAmbientOcclusion", EditConditionHides, UIMin = "0", UIMax = "10.0", ClampMin = "0", ClampMax = "100.0"))
 	double OcclusionBlurRadius = 2.25;
 
 	/** Pixel Radius of Gaussian Blur Kernel applied to Curvature map (0 will apply no blur) */
-	UPROPERTY(EditAnywhere, Category = SmoothedCurvatureSettings, meta = (EditCondition = "bSmoothedCurvature", EditConditionHides, UIMin = "0", UIMax = "10.0", ClampMin = "0", ClampMax = "100.0"))
+	UPROPERTY(EditAnywhere, Category = SmoothedCurvatureSettings, meta = (EditCondition = "BakeTextureType == ETextureType::ThicknessAndSurfaceAttributes && bSmoothedCurvature", EditConditionHides, UIMin = "0", UIMax = "10.0", ClampMin = "0", ClampMax = "100.0"))
 	double CurvatureBlurRadius = 2.25;
 
 	/** Voxel resolution of smoothed shape representation */
-	UPROPERTY(EditAnywhere, Category = SmoothedCurvatureSettings, meta = (EditCondition = "bSmoothedCurvature", EditConditionHides, UIMin = "8", UIMax = "512", ClampMin = "4", ClampMax = "1024"))
+	UPROPERTY(EditAnywhere, Category = SmoothedCurvatureSettings, meta = (EditCondition = "BakeTextureType == ETextureType::ThicknessAndSurfaceAttributes && bSmoothedCurvature", EditConditionHides, UIMin = "8", UIMax = "512", ClampMin = "4", ClampMax = "1024"))
 	int VoxelResolution = 128;
 
 	/** Amount of smoothing iterations to apply before computing curvature */
-	UPROPERTY(EditAnywhere, Category = SmoothedCurvatureSettings, meta = (EditCondition = "bSmoothedCurvature", EditConditionHides, UIMin = "2", UIMax = "100", ClampMin = "2", ClampMax = "1000"))
+	UPROPERTY(EditAnywhere, Category = SmoothedCurvatureSettings, meta = (EditCondition = "BakeTextureType == ETextureType::ThicknessAndSurfaceAttributes && bSmoothedCurvature", EditConditionHides, UIMin = "2", UIMax = "100", ClampMin = "2", ClampMax = "1000"))
 	int SmoothingIterations = 10;
 
 	/** Distance to search for correspondence between fractured shape and smoothed shape, as factor of voxel size */
-	UPROPERTY(EditAnywhere, Category = SmoothedCurvatureSettings, meta = (EditCondition = "bSmoothedCurvature", EditConditionHides, UIMin = "2", UIMax = "10.0", ClampMin = "1", ClampMax = "100.0"))
+	UPROPERTY(EditAnywhere, Category = SmoothedCurvatureSettings, meta = (EditCondition = "BakeTextureType == ETextureType::ThicknessAndSurfaceAttributes && bSmoothedCurvature", EditConditionHides, UIMin = "2", UIMax = "10.0", ClampMin = "1", ClampMax = "100.0"))
 	double ThicknessFactor = 4;
 
 	/** Curvatures in the range [-MaxCurvature, MaxCurvature] will be mapped from [0,1]. Values outside that range will be clamped. */
-	UPROPERTY(EditAnywhere, Category = SmoothedCurvatureSettings, meta = (EditCondition = "bSmoothedCurvature", EditConditionHides, UIMin = ".01", UIMax = "1", ClampMin = ".0001", ClampMax = "10"))
+	UPROPERTY(EditAnywhere, Category = SmoothedCurvatureSettings, meta = (EditCondition = "BakeTextureType == ETextureType::ThicknessAndSurfaceAttributes && bSmoothedCurvature", EditConditionHides, UIMin = ".01", UIMax = "1", ClampMin = ".0001", ClampMax = "10"))
 	double MaxCurvature = .1;
 
 	/** Whether to use the absolute value of the Z-Component of the normal */
-	UPROPERTY(EditAnywhere, Category = ZNormalSettings, meta = (EditCondition = "bZNormal", EditConditionHides))
+	UPROPERTY(EditAnywhere, Category = ZNormalSettings, meta = (EditCondition = "BakeTextureType == ETextureType::ThicknessAndSurfaceAttributes && bZNormal", EditConditionHides))
 	bool bUseAbsoluteValue = true;
 
 };
diff --git a/Engine/Plugins/Experimental/PlanarCutPlugin/Source/PlanarCut/Private/FractureAutoUV.cpp b/Engine/Plugins/Experimental/PlanarCutPlugin/Source/PlanarCut/Private/FractureAutoUV.cpp
--- a/Engine/Plugins/Experimental/PlanarCutPlugin/Source/PlanarCut/Private/FractureAutoUV.cpp
+++ b/Engine/Plugins/Experimental/PlanarCutPlugin/Source/PlanarCut/Private/FractureAutoUV.cpp
@@ -504,13 +504,11 @@ void TextureInternalSurfaces(
 	int AmbientIdx = BakeAttributes.IndexOf((int)EBakeAttributes::AmbientOcclusion);
 	int CurvatureIdx = BakeAttributes.IndexOf((int)EBakeAttributes::Curvature);
 	int NormalZIdx = BakeAttributes.IndexOf((int)EBakeAttributes::NormalZ);
-	int PosZIdx = BakeAttributes.IndexOf((int)EBakeAttributes::PositionZ);
-	
-	if (OutsideSpatial.GetBoundingBox().Depth() == 0)
-	{
-		PosZIdx = -1; // if everything has same Z, can just leave PosZ values as default value
-	}
 
+	FAxisAlignedBox3d Box = OutsideSpatial.GetBoundingBox();
+	int PosXIdx = Box.Width() > 0 ? BakeAttributes.IndexOf((int)EBakeAttributes::PositionX) : -1;
+	int PosYIdx = Box.Height() > 0 ? BakeAttributes.IndexOf((int)EBakeAttributes::PositionY) : -1;
+	int PosZIdx = Box.Depth() > 0 ? BakeAttributes.IndexOf((int)EBakeAttributes::PositionZ) : -1;
 
 	bool bNeedsDynamicMeshes = AmbientIdx > -1 || CurvatureIdx > -1;
 
@@ -675,7 +673,7 @@ void TextureInternalSurfaces(
 
 	ParallelFor(OccupancyMap.Dimensions.GetHeight(),
 		[&AttributeSettings, &TextureOut, &UVMesh, &OccupancyMap, &InsideMesh, &OutsideSpatial,
-		 &DistanceToExternalIdx, &AmbientIdx, &NormalZIdx, &PosZIdx](int32 Y)
+		 &DistanceToExternalIdx, &AmbientIdx, &NormalZIdx, &PosXIdx, &PosYIdx, &PosZIdx](int32 Y)
 	{
 		for (int32 X = 0; X < OccupancyMap.Dimensions.GetWidth(); X++)
 		{
@@ -697,7 +695,7 @@ void TextureInternalSurfaces(
 					Normal = InsideMesh.GetInterpolatedNormal(TID, Bary);
 				}
 
-				if (DistanceToExternalIdx > -1 || PosZIdx > -1)
+				if (DistanceToExternalIdx > -1 || PosXIdx > -1 || PosYIdx > -1 || PosZIdx > -1)
 				{
 					FTriangle3d Tri;
 					InsideMesh.GetTriVertices(TID, Tri.V[0], Tri.V[1], Tri.V[2]);
@@ -711,13 +709,19 @@ void TextureInternalSurfaces(
 						checkSlow(FMath::IsFinite(PercentDistance));
 						OutColor[DistanceToExternalIdx] = PercentDistance;
 					}
-					if (PosZIdx > -1)
+					auto SetSpatial = [&OutsideSpatial, &InsidePoint, &OutColor](int TargetIdx, int Dim)
 					{
-						double MinZ = OutsideSpatial.GetBoundingBox().Min.Z, MaxZ = OutsideSpatial.GetBoundingBox().Max.Z;
-						checkSlow(MinZ != MaxZ);
-						float PercentHeight = (float)(InsidePoint.Z - MinZ) / (MaxZ - MinZ);
-						OutColor[PosZIdx] = PercentHeight;
-					}
+						if (TargetIdx > -1)
+						{
+							double Min = OutsideSpatial.GetBoundingBox().Min[Dim], Max = OutsideSpatial.GetBoundingBox().Max[Dim];
+							checkSlow(Min != Max);
+							float PercentAlong = (float)(InsidePoint[Dim] - Min) / (Max - Min);
+							OutColor[TargetIdx] = PercentAlong;
+						}
+					};
+					SetSpatial(PosXIdx, 0);
+					SetSpatial(PosYIdx, 1);
+					SetSpatial(PosZIdx, 2);
 				}
 				if (NormalZIdx > -1)
 				{
diff --git a/Engine/Plugins/Experimental/PlanarCutPlugin/Source/PlanarCut/Public/FractureAutoUV.h b/Engine/Plugins/Experimental/PlanarCutPlugin/Source/PlanarCut/Public/FractureAutoUV.h
--- a/Engine/Plugins/Experimental/PlanarCutPlugin/Source/PlanarCut/Public/FractureAutoUV.h
+++ b/Engine/Plugins/Experimental/PlanarCutPlugin/Source/PlanarCut/Public/FractureAutoUV.h
@@ -36,6 +36,8 @@ enum class EBakeAttributes : int32
 	AmbientOcclusion,
 	Curvature,
 	NormalZ,
+	PositionX,
+	PositionY,
 	PositionZ
 };
 
