allow overwriting files in autouv, so we can replace the existing texture

#rnx

diff --git a/Engine/Plugins/Experimental/ChaosEditor/ChaosEditor.uplugin b/Engine/Plugins/Experimental/ChaosEditor/ChaosEditor.uplugin
--- a/Engine/Plugins/Experimental/ChaosEditor/ChaosEditor.uplugin
+++ b/Engine/Plugins/Experimental/ChaosEditor/ChaosEditor.uplugin
@@ -33,6 +33,10 @@
 		{
 			"Name": "MeshModelingToolset",
 			"Enabled": true
+		},
+		{
+			"Name": "EditorScriptingUtilities",
+			"Enabled": true
 		}
 	]
 }
\ No newline at end of file
diff --git a/Engine/Plugins/Experimental/ChaosEditor/Source/FractureEditor/FractureEditor.Build.cs b/Engine/Plugins/Experimental/ChaosEditor/Source/FractureEditor/FractureEditor.Build.cs
--- a/Engine/Plugins/Experimental/ChaosEditor/Source/FractureEditor/FractureEditor.Build.cs
+++ b/Engine/Plugins/Experimental/ChaosEditor/Source/FractureEditor/FractureEditor.Build.cs
@@ -43,6 +43,7 @@ public class FractureEditor : ModuleRules
 				"SlateCore",
 				"InputCore",
 				"EditorFramework",
+				"EditorScriptingUtilities",
 				"ToolMenus",
 				"UnrealEd",
 				"LevelEditor",
diff --git a/Engine/Plugins/Experimental/ChaosEditor/Source/FractureEditor/Private/FractureToolUV.cpp b/Engine/Plugins/Experimental/ChaosEditor/Source/FractureEditor/Private/FractureToolUV.cpp
--- a/Engine/Plugins/Experimental/ChaosEditor/Source/FractureEditor/Private/FractureToolUV.cpp
+++ b/Engine/Plugins/Experimental/ChaosEditor/Source/FractureEditor/Private/FractureToolUV.cpp
@@ -11,6 +11,7 @@
 #include "AssetUtils/Texture2DBuilder.h"
 #include "AssetToolsModule.h"
 #include "AssetRegistryModule.h"
+#include "EditorAssetLibrary.h"
 // for content-browser things
 #include "ContentBrowserModule.h"
 #include "IContentBrowserSingleton.h"
@@ -121,7 +122,8 @@ TArray<FFractureToolContext> UFractureToolAutoUV::GetFractureToolContexts() cons
 	return Contexts;
 }
 
-bool UFractureToolAutoUV::SaveGeneratedTexture(UTexture2D* GeneratedTexture, FString ObjectBaseName, const UObject* RelativeToAsset, bool bPromptToSave)
+
+bool UFractureToolAutoUV::SaveGeneratedTexture(UTexture2D* GeneratedTexture, FString ObjectBaseName, const UObject* RelativeToAsset, bool bPromptToSave, bool bAllowReplace)
 {
 	check(RelativeToAsset);
 	check(GeneratedTexture);
@@ -141,14 +143,15 @@ bool UFractureToolAutoUV::SaveGeneratedTexture(UTexture2D* GeneratedTexture, FSt
 
 		FString UseDefaultAssetName = ObjectBaseName;
 		FString CurrentPath = PackageFolderPath;
-		if (CurrentPath.IsEmpty() == false)
+		if (!CurrentPath.IsEmpty() && !bAllowReplace)
 		{
 			FAssetToolsModule& AssetToolsModule = FModuleManager::LoadModuleChecked<FAssetToolsModule>("AssetTools");
 			FString UnusedPackageName;
-			AssetToolsModule.Get().CreateUniqueAssetName(PackageFolderPath + TEXT("/") + ObjectBaseName, TEXT(""), UnusedPackageName, UseDefaultAssetName);
+			AssetToolsModule.Get().CreateUniqueAssetName(FPaths::Combine(PackageFolderPath, ObjectBaseName), TEXT(""), UnusedPackageName, UseDefaultAssetName);
 		}
 
 		FSaveAssetDialogConfig Config;
+		Config.ExistingAssetPolicy = ESaveAssetDialogExistingAssetPolicy::AllowButWarn;
 		Config.DefaultAssetName = UseDefaultAssetName;
 		Config.DialogTitleOverride = LOCTEXT("GenerateTexture2DAssetPathDialogWarning", "Choose Folder Path and Name for New Asset. Cancel to Discard New Asset.");
 		Config.DefaultPath = CurrentPath;
@@ -165,12 +168,24 @@ bool UFractureToolAutoUV::SaveGeneratedTexture(UTexture2D* GeneratedTexture, FSt
 		}
 	}
 
+	FString NewAssetPath = FPaths::Combine(PackageFolderPath, ObjectBaseName);
+	if (bAllowReplace)
+	{
+		// Modifying the static mesh in place. Delete existing asset so that we can have a clean duplicate
+		bool bNewAssetExists = UEditorAssetLibrary::DoesAssetExist(NewAssetPath);
+		if (bNewAssetExists)
+		{
+			bool bDeleteOK = UEditorAssetLibrary::DeleteAsset(NewAssetPath);
+			ensure(bDeleteOK);
+		}
+	}
+
 	// create new package
 	FString UniqueAssetName;
 	FString UniquePackageName;
 
 	FAssetToolsModule& AssetToolsModule = FModuleManager::LoadModuleChecked<FAssetToolsModule>("AssetTools");
-	AssetToolsModule.Get().CreateUniqueAssetName(PackageFolderPath + TEXT("/") + ObjectBaseName, TEXT(""), UniquePackageName, UniqueAssetName);
+	AssetToolsModule.Get().CreateUniqueAssetName(NewAssetPath, TEXT(""), UniquePackageName, UniqueAssetName);
 
 	UPackage* AssetPackage = CreatePackage(*UniquePackageName);
 
@@ -265,7 +280,13 @@ int32 UFractureToolAutoUV::ExecuteFracture(const FFractureToolContext& FractureC
 		// choose default texture name based on corresponding geometry collection name
 		FString BaseName = FractureContext.GetFracturedGeometryCollection()->GetName();
 		FTexture2DBuilder::CopyPlatformDataToSourceData(AutoUVSettings->Result, FTexture2DBuilder::ETextureType::Color);
-		SaveGeneratedTexture(AutoUVSettings->Result, FString::Printf(TEXT("%s_AutoUV"), *BaseName), FractureContext.GetFracturedGeometryCollection(), AutoUVSettings->bPromptToSave);
+		FString Suffix = "_AutoUV";
+		if (AutoUVSettings->BakeTextureType == ETextureType::SpatialGradients)
+		{
+			Suffix = "_AutoUV_Spatial";
+		}
+		Suffix = FPaths::MakeValidFileName(Suffix);
+		SaveGeneratedTexture(AutoUVSettings->Result, FString::Printf(TEXT("%s%s"), *BaseName, *Suffix), FractureContext.GetFracturedGeometryCollection(), AutoUVSettings->bPromptToSave, AutoUVSettings->bReplaceExisting);
 	}
 
 	return INDEX_NONE;
diff --git a/Engine/Plugins/Experimental/ChaosEditor/Source/FractureEditor/Private/FractureToolUV.h b/Engine/Plugins/Experimental/ChaosEditor/Source/FractureEditor/Private/FractureToolUV.h
--- a/Engine/Plugins/Experimental/ChaosEditor/Source/FractureEditor/Private/FractureToolUV.h
+++ b/Engine/Plugins/Experimental/ChaosEditor/Source/FractureEditor/Private/FractureToolUV.h
@@ -62,6 +62,10 @@ public:
 	UPROPERTY(EditAnywhere, Category = MapSettings)
 	bool bPromptToSave = true;
 
+	/** Whether to allow the new texture to overwrite an existing texture */
+	UPROPERTY(EditAnywhere, Category = MapSettings)
+	bool bReplaceExisting = true;
+
 	/** Which standard set of texture channels to bake */
 	UPROPERTY(EditAnywhere, Category = AttributesToBake)
 	ETextureType BakeTextureType = ETextureType::ThicknessAndSurfaceAttributes;
@@ -161,5 +165,5 @@ protected:
 	UPROPERTY(EditAnywhere, Category = Slicing)
 	TObjectPtr<UFractureAutoUVSettings> AutoUVSettings;
 
-	bool SaveGeneratedTexture(UTexture2D* GeneratedTexture, FString ObjectBaseName, const UObject* RelativeToAsset, bool bPromptToSave);
+	bool SaveGeneratedTexture(UTexture2D* GeneratedTexture, FString ObjectBaseName, const UObject* RelativeToAsset, bool bPromptToSave, bool bAllowReplace);
 };
