change fracture convert tool defaults to not center pivots

#rnx

diff --git a/Engine/Plugins/Experimental/ChaosEditor/Source/FractureEditor/Private/FractureToolConvert.h b/Engine/Plugins/Experimental/ChaosEditor/Source/FractureEditor/Private/FractureToolConvert.h
--- a/Engine/Plugins/Experimental/ChaosEditor/Source/FractureEditor/Private/FractureToolConvert.h
+++ b/Engine/Plugins/Experimental/ChaosEditor/Source/FractureEditor/Private/FractureToolConvert.h
@@ -31,7 +31,7 @@ public:
 
 	/** Whether to center the pivot for each mesh, or use the pivot from the geometry collection */
 	UPROPERTY(EditAnywhere, Category = GeometrySettings)
-	bool bCenterPivots = true;
+	bool bCenterPivots = false;
 };
 
 
