fix radial fracture tool point distribution to be planar for all normal directions

#rnx

diff --git a/Engine/Plugins/Experimental/ChaosEditor/Source/FractureEditor/Private/FractureToolRadial.cpp b/Engine/Plugins/Experimental/ChaosEditor/Source/FractureEditor/Private/FractureToolRadial.cpp
--- a/Engine/Plugins/Experimental/ChaosEditor/Source/FractureEditor/Private/FractureToolRadial.cpp
+++ b/Engine/Plugins/Experimental/ChaosEditor/Source/FractureEditor/Private/FractureToolRadial.cpp
@@ -47,7 +47,8 @@ TArray<UObject*> UFractureToolRadial::GetSettingsObjects() const
 
 void UFractureToolRadial::GenerateVoronoiSites(const FFractureToolContext& Context, TArray<FVector>& Sites)
 {
- 	float RadialStep = RadialSettings->Radius / RadialSettings->RadialSteps;
+ 	const FVector::FReal RadialStep = RadialSettings->Radius / RadialSettings->RadialSteps;
+	const FVector::FReal AngularStep = 2 * PI / RadialSettings->AngularSteps;
 
 	FBox Bounds = Context.GetWorldBounds();
 	const FVector Center(Bounds.GetCenter() + RadialSettings->Center);
@@ -55,19 +56,17 @@ void UFractureToolRadial::GenerateVoronoiSites(const FFractureToolContext& Conte
 	FRandomStream RandStream(Context.GetSeed());
 	FVector UpVector(RadialSettings->Normal);
 	UpVector.Normalize();
-	FVector PerpVector(UpVector[2], UpVector[0], UpVector[1]);
+	FVector BasisX, BasisY;
+	UpVector.FindBestAxisVectors(BasisX, BasisY);
 
-	for (int32 ii = 1; ii < RadialSettings->RadialSteps; ++ii)
+	FVector::FReal Len = RadialStep * .5;
+	for (int32 ii = 0; ii < RadialSettings->RadialSteps; ++ii, Len += RadialStep)
 	{
-		FVector PositionVector(PerpVector * RadialStep * ii);
-
-		float AngularStep = 360.f / RadialSettings->AngularSteps;
-		PositionVector = PositionVector.RotateAngleAxis(RadialSettings->AngleOffset * ii, UpVector);
-
-		for (int32 kk = 0; kk < RadialSettings->AngularSteps; ++kk)
+		FVector::FReal Angle = RadialSettings->AngleOffset;
+		for (int32 kk = 0; kk < RadialSettings->AngularSteps; ++kk, Angle += AngularStep)
 		{
-			PositionVector = PositionVector.RotateAngleAxis(AngularStep , UpVector);
-			Sites.Emplace(Center + PositionVector + (RandStream.VRand() * RandStream.FRand() * RadialSettings->Variability));
+			FVector RotatingOffset = Len * (FMath::Cos(Angle) * BasisX + FMath::Sin(Angle) * BasisY);
+			Sites.Emplace(Center + RotatingOffset + (RandStream.VRand() * RandStream.FRand() * RadialSettings->Variability));
 		}
 	}
 }
