don't generate lightmap UVs on the gc->staticmesh conversion mesh; geometrycollections don't carry lightmaps and these will typically be processed then re-added to a geometry collection

#rnx

diff --git a/Engine/Plugins/Experimental/ChaosEditor/Source/FractureEditor/Private/FractureToolConvert.cpp b/Engine/Plugins/Experimental/ChaosEditor/Source/FractureEditor/Private/FractureToolConvert.cpp
--- a/Engine/Plugins/Experimental/ChaosEditor/Source/FractureEditor/Private/FractureToolConvert.cpp
+++ b/Engine/Plugins/Experimental/ChaosEditor/Source/FractureEditor/Private/FractureToolConvert.cpp
@@ -149,6 +149,8 @@ AActor* CreateMeshAsset(FGeometryCollection& Collection, const FTransform& Colle
 	// normals and tangents should carry over from the geometry collection
 	NewStaticMesh->GetSourceModel(0).BuildSettings.bRecomputeNormals = false;
 	NewStaticMesh->GetSourceModel(0).BuildSettings.bRecomputeTangents = false;
+	NewStaticMesh->GetSourceModel(0).BuildSettings.bGenerateLightmapUVs = false;
+
 	FMeshDescription* OutputMeshDescription = NewStaticMesh->CreateMeshDescription(0);
 
 	NewStaticMesh->CreateBodySetup();
