LWC: ShooterGame build fixes

diff --git a/Engine/Plugins/Experimental/ChaosEditor/Source/FractureEditor/Private/AutoClusterFracture.cpp b/Engine/Plugins/Experimental/ChaosEditor/Source/FractureEditor/Private/AutoClusterFracture.cpp
--- a/Engine/Plugins/Experimental/ChaosEditor/Source/FractureEditor/Private/AutoClusterFracture.cpp
+++ b/Engine/Plugins/Experimental/ChaosEditor/Source/FractureEditor/Private/AutoClusterFracture.cpp
@@ -434,7 +434,7 @@ bool UAutoClusterFractureCommand::HasPath(int32 TransformIndexStart, int32 Trans
 
 float UAutoClusterFractureCommand::GetClosestDistance(const FBox& A, const FBox& B)
 {
-	float Dist[8] = {
+	FVector::FReal Dist[8] = {
 		B.ComputeSquaredDistanceToPoint(FVector(A.Min.X, A.Min.Y, A.Min.Z)),
 		B.ComputeSquaredDistanceToPoint(FVector(A.Min.X, A.Max.Y, A.Min.Z)),
 		B.ComputeSquaredDistanceToPoint(FVector(A.Max.X, A.Min.Y, A.Min.Z)),
@@ -446,7 +446,7 @@ float UAutoClusterFractureCommand::GetClosestDistance(const FBox& A, const FBox&
 		B.ComputeSquaredDistanceToPoint(FVector(A.Max.X, A.Max.Y, A.Max.Z)) 
 	};
 
-	float Distance = FLT_MAX;
+	FVector::FReal Distance = FLT_MAX;
 
 	for (int32 ii = 0; ii < 8; ++ii)
 	{
@@ -456,7 +456,7 @@ float UAutoClusterFractureCommand::GetClosestDistance(const FBox& A, const FBox&
 		}
 	}
 
-	return Distance;
+	return (float)Distance;			// LWC_TODO: Precision loss
 }
 
 
diff --git a/Engine/Source/Runtime/Core/Public/Math/SHMath.h b/Engine/Source/Runtime/Core/Public/Math/SHMath.h
--- a/Engine/Source/Runtime/Core/Public/Math/SHMath.h
+++ b/Engine/Source/Runtime/Core/Public/Math/SHMath.h
@@ -351,7 +351,7 @@ public:
 		// we know the Vector will be (0,0,+1) or (0,0,-1).
 		if (FMath::Abs(Vector.X) > KINDA_SMALL_NUMBER || FMath::Abs(Vector.Y) > KINDA_SMALL_NUMBER)
 		{
-			const float	Phi = FMath::Atan2(Vector.Y, Vector.X);
+			const float	Phi = (float)FMath::Atan2(Vector.Y, Vector.X);
 
 			for (int32 BandIndex = 1; BandIndex < TSHVector::MaxSHOrder; BandIndex++)
 			{
