Add gizmo transform widget support and connect to radial fracture.
Includes related changes to support this:
 - Change from FEdMode to a UEdMode to access the contexts the gizmo needs (this may not have been strictly needed due to the addition of a global context, but the gizmo being in that global context may be in flux, and the UEdMode is the new way of doing things generally)
 - Add virtual functions to the modal tool for tool setup, shutdown, and selection change (so gizmo can be created/removed/updated as needed)

#rnx

diff --git a/Engine/Plugins/Experimental/ChaosEditor/Source/FractureEditor/FractureEditor.Build.cs b/Engine/Plugins/Experimental/ChaosEditor/Source/FractureEditor/FractureEditor.Build.cs
--- a/Engine/Plugins/Experimental/ChaosEditor/Source/FractureEditor/FractureEditor.Build.cs
+++ b/Engine/Plugins/Experimental/ChaosEditor/Source/FractureEditor/FractureEditor.Build.cs
@@ -26,7 +26,9 @@ public class FractureEditor : ModuleRules
 			new string[]
 			{
 				"Core",
-                "Voronoi"
+                "Voronoi",
+				"InteractiveToolsFramework",
+				"EditorInteractiveToolsFramework"
 
 				// ... add other public dependencies that you statically link with here ...
 			}
diff --git a/Engine/Plugins/Experimental/ChaosEditor/Source/FractureEditor/Private/FractureEditor.cpp b/Engine/Plugins/Experimental/ChaosEditor/Source/FractureEditor/Private/FractureEditor.cpp
--- a/Engine/Plugins/Experimental/ChaosEditor/Source/FractureEditor/Private/FractureEditor.cpp
+++ b/Engine/Plugins/Experimental/ChaosEditor/Source/FractureEditor/Private/FractureEditor.cpp
@@ -17,12 +17,6 @@ void FFractureEditorModule::StartupModule()
 
 	// This code will execute after your module is loaded into memory; the exact timing is specified in the .uplugin file per-module
 
-	FEditorModeRegistry::Get().RegisterMode<FFractureEditorMode>(
-		FFractureEditorMode::EM_FractureEditorModeId, 
-		LOCTEXT("FractureEditorModeName", "Fracture"), 
-		FSlateIcon("FractureEditorStyle", "LevelEditor.FractureMode", "LevelEditor.FractureMode.Small"),
-		true
-		);
 
 	ToolBarExtensibilityManager = MakeShareable(new FExtensibilityManager);
 
@@ -36,7 +30,6 @@ void FFractureEditorModule::ShutdownModule()
 {
 	// This function may be called during shutown to clean up your module.  For modules that support dynamic reloading,
 	// we call this function before unloading the module.
-	FEditorModeRegistry::Get().UnregisterMode(FFractureEditorMode::EM_FractureEditorModeId);
 
 	
 	FFractureEditorCommands::Unregister();
diff --git a/Engine/Plugins/Experimental/ChaosEditor/Source/FractureEditor/Private/FractureEditorMode.cpp b/Engine/Plugins/Experimental/ChaosEditor/Source/FractureEditor/Private/FractureEditorMode.cpp
--- a/Engine/Plugins/Experimental/ChaosEditor/Source/FractureEditor/Private/FractureEditorMode.cpp
+++ b/Engine/Plugins/Experimental/ChaosEditor/Source/FractureEditor/Private/FractureEditorMode.cpp
@@ -24,36 +24,43 @@
 #include "EditorViewportClient.h"
 #include "ScopedTransaction.h"
 #include "GeometryCollection/GeometryCollectionAlgo.h"
+#include "EdModeInteractiveToolsContext.h"
+#include "BaseGizmos/TransformGizmoUtil.h"
+
+
+#include "UnrealEdGlobals.h"
+#include "EditorModeManager.h"
 
 #define LOCTEXT_NAMESPACE "FFractureEditorModeToolkit"
 
-const FEditorModeID FFractureEditorMode::EM_FractureEditorModeId = TEXT("EM_FractureEditorMode");
+const FEditorModeID UFractureEditorMode::EM_FractureEditorModeId = TEXT("EM_FractureEditorMode");
 
-FFractureEditorMode::FFractureEditorMode()
+UFractureEditorMode::UFractureEditorMode()
 {
+	Info = FEditorModeInfo(
+		EM_FractureEditorModeId,
+		LOCTEXT("FractureToolsEditorModeName", "Fracture"),
+		FSlateIcon("FractureEditorStyle", "LevelEditor.FractureMode", "LevelEditor.FractureMode.Small"),
+		true);
 }
 
-FFractureEditorMode::~FFractureEditorMode()
+UFractureEditorMode::~UFractureEditorMode()
 {
 
 }
 
-void FFractureEditorMode::Enter()
+void UFractureEditorMode::Enter()
 {
-	FEdMode::Enter();
+	Super::Enter();
 
 	GEditor->RegisterForUndo(this);
 
-	if (!Toolkit.IsValid() && UsesToolkits())
-	{
-		Toolkit = MakeShareable(new FFractureEditorModeToolkit);
-		Toolkit->Init(Owner->GetToolkitHost());
-	}
-
 	FLevelEditorModule& LevelEditorModule = FModuleManager::LoadModuleChecked<FLevelEditorModule>("LevelEditor");
-	LevelEditorModule.OnActorSelectionChanged().AddRaw(this, &FFractureEditorMode::OnActorSelectionChanged);
+	LevelEditorModule.OnActorSelectionChanged().AddUObject(this, &UFractureEditorMode::OnActorSelectionChanged);
+
+	FCoreUObjectDelegates::OnPackageReloaded.AddUObject(this, &UFractureEditorMode::HandlePackageReloaded);
 
-	FCoreUObjectDelegates::OnPackageReloaded.AddSP(this, &FFractureEditorMode::HandlePackageReloaded);
+	UE::TransformGizmoUtil::RegisterTransformGizmoContextObject(ToolsContext.Get());
 	
 	// Get initial geometry component selection from currently selected actors when we enter the mode
 	USelection* SelectedActors = GEditor->GetSelectedActors();
@@ -70,16 +77,21 @@ void FFractureEditorMode::Enter()
 
 }
 
-void FFractureEditorMode::Exit()
+void UFractureEditorMode::Exit()
 {
 	GEditor->UnregisterForUndo(this);
 
+	// TODO: cannot deregister currently because if another mode is also registering, its Enter()
+	// will be called before our Exit(); add the below line back after this bug is fixed
+	//FractureGizmoHelper->DeregisterGizmosWithManager(ToolsContext->ToolManager);
+
 	// Empty the geometry component selection set
 	TArray<UObject*> SelectedObjects;
 	OnActorSelectionChanged(SelectedObjects, false);
 
 	if (Toolkit.IsValid())
 	{
+		static_cast<FFractureEditorModeToolkit*>(Toolkit.Get())->Shutdown();
 		FToolkitManager::Get().CloseToolkit(Toolkit.ToSharedRef());
 		Toolkit.Reset();
 	}
@@ -92,15 +104,10 @@ void FFractureEditorMode::Exit()
 	}
 
 	// Call base Exit method to ensure proper cleanup
-	FEdMode::Exit();
+	Super::Exit();
 }
 
-void FFractureEditorMode::AddReferencedObjects(FReferenceCollector& Collector)
-{
-	Collector.AddReferencedObjects(SelectedGeometryComponents);
-}
-
-bool FFractureEditorMode::MatchesContext(const FTransactionContext& InContext, const TArray<TPair<UObject *, FTransactionObjectEvent>>& TransactionObjectContexts) const
+bool UFractureEditorMode::MatchesContext(const FTransactionContext& InContext, const TArray<TPair<UObject *, FTransactionObjectEvent>>& TransactionObjectContexts) const
 {
 	if (InContext.Context == FractureTransactionContexts::SelectBoneContext)
 	{
@@ -110,19 +117,19 @@ bool FFractureEditorMode::MatchesContext(const FTransactionContext& InContext, c
 	return false;
 }
 
-void FFractureEditorMode::PostUndo(bool bSuccess)
+void UFractureEditorMode::PostUndo(bool bSuccess)
 {
 	OnUndoRedo();
 }
 
-void FFractureEditorMode::PostRedo(bool bSuccess)
+void UFractureEditorMode::PostRedo(bool bSuccess)
 {
 	OnUndoRedo();
 }
 
-void FFractureEditorMode::Render(const FSceneView* View, FViewport* Viewport, FPrimitiveDrawInterface* PDI)
+void UFractureEditorMode::Render(const FSceneView* View, FViewport* Viewport, FPrimitiveDrawInterface* PDI)
 {
-	FEdMode::Render(View, Viewport, PDI);
+	Super::Render(View, Viewport, PDI);
 
 	FFractureEditorModeToolkit* FractureToolkit = (FFractureEditorModeToolkit*)Toolkit.Get();
  
@@ -135,12 +142,21 @@ void FFractureEditorMode::Render(const FSceneView* View, FViewport* Viewport, FP
 
 }
 
-bool FFractureEditorMode::UsesToolkits() const
+
+void UFractureEditorMode::CreateToolkit()
+{
+	if (!Toolkit.IsValid() && UsesToolkits())
+	{
+		Toolkit = MakeShareable(new FFractureEditorModeToolkit);
+	}
+}
+
+bool UFractureEditorMode::UsesToolkits() const
 {
 	return true;
 }
 
-bool FFractureEditorMode::InputKey(FEditorViewportClient* ViewportClient, FViewport* Viewport, FKey Key, EInputEvent Event)
+bool UFractureEditorMode::InputKey(FEditorViewportClient* ViewportClient, FViewport* Viewport, FKey Key, EInputEvent Event)
 {
 
 	bool bHandled = false;
@@ -153,7 +169,7 @@ bool FFractureEditorMode::InputKey(FEditorViewportClient* ViewportClient, FViewp
 	return bHandled;
 }
 
-bool FFractureEditorMode::HandleClick(FEditorViewportClient* InViewportClient, HHitProxy* HitProxy, const FViewportClick& Click)
+bool UFractureEditorMode::HandleClick(FEditorViewportClient* InViewportClient, HHitProxy* HitProxy, const FViewportClick& Click)
 {
 	if (HitProxy && HitProxy->IsA(HGeometryCollectionBone::StaticGetType()))
 	{
@@ -181,13 +197,13 @@ bool FFractureEditorMode::HandleClick(FEditorViewportClient* InViewportClient, H
 
 }
 
-bool FFractureEditorMode::BoxSelect(FBox& InBox, bool InSelect /*= true*/)
+bool UFractureEditorMode::BoxSelect(FBox& InBox, bool InSelect /*= true*/)
 {
 	FConvexVolume BoxVolume(GetVolumeFromBox(InBox));
 	return FrustumSelect(BoxVolume, nullptr, InSelect);
 }
 
-bool FFractureEditorMode::FrustumSelect(const FConvexVolume& InFrustum, FEditorViewportClient* InViewportClient, bool InSelect /*= true*/)
+bool UFractureEditorMode::FrustumSelect(const FConvexVolume& InFrustum, FEditorViewportClient* InViewportClient, bool InSelect /*= true*/)
 {
 	bool bStrictDragSelection = GetDefault<ULevelEditorViewportSettings>()->bStrictBoxSelection;
 	bool bSelectedBones = false;
@@ -255,7 +271,7 @@ bool FFractureEditorMode::FrustumSelect(const FConvexVolume& InFrustum, FEditorV
 	return bSelectedBones;
 }
 
-bool FFractureEditorMode::ComputeBoundingBoxForViewportFocus(AActor* Actor, UPrimitiveComponent* PrimitiveComponent, FBox& InOutBox) const
+bool UFractureEditorMode::ComputeBoundingBoxForViewportFocus(AActor* Actor, UPrimitiveComponent* PrimitiveComponent, FBox& InOutBox) const
 {
 	UGeometryCollectionComponent* GeometryCollectionComponent = Cast<UGeometryCollectionComponent>(PrimitiveComponent);
 	if (GeometryCollectionComponent && GeometryCollectionComponent->GetSelectedBones().Num() > 0 && SelectedGeometryComponents.Contains(GeometryCollectionComponent))
@@ -281,7 +297,7 @@ bool FFractureEditorMode::ComputeBoundingBoxForViewportFocus(AActor* Actor, UPri
 	return false;
 }
 
-bool FFractureEditorMode::GetPivotForOrbit(FVector& OutPivot) const
+bool UFractureEditorMode::GetPivotForOrbit(FVector& OutPivot) const
 {
 	if (CustomOrbitPivot.IsSet())
 	{
@@ -291,7 +307,7 @@ bool FFractureEditorMode::GetPivotForOrbit(FVector& OutPivot) const
 	return CustomOrbitPivot.IsSet();
 }
 
-void FFractureEditorMode::OnUndoRedo()
+void UFractureEditorMode::OnUndoRedo()
 {
 	for (UGeometryCollectionComponent* SelectedComp : SelectedGeometryComponents)
 	{
@@ -301,7 +317,7 @@ void FFractureEditorMode::OnUndoRedo()
 	}
 }
 
-void FFractureEditorMode::OnActorSelectionChanged(const TArray<UObject*>& NewSelection, bool bForceRefresh)
+void UFractureEditorMode::OnActorSelectionChanged(const TArray<UObject*>& NewSelection, bool bForceRefresh)
 {
 	CustomOrbitPivot.Reset();
 	TSet<UGeometryCollectionComponent*> NewGeomSelection;
@@ -357,7 +373,7 @@ void FFractureEditorMode::OnActorSelectionChanged(const TArray<UObject*>& NewSel
 	}
 }
 
-void FFractureEditorMode::GetActorGlobalBounds(TArrayView<UGeometryCollectionComponent*> GeometryComponents, TMap<int32, FBox> &BoundsToBone) const
+void UFractureEditorMode::GetActorGlobalBounds(TArrayView<UGeometryCollectionComponent*> GeometryComponents, TMap<int32, FBox> &BoundsToBone) const
 {
 	for (UGeometryCollectionComponent* GeometryCollectionComponent : GeometryComponents)
 	{
@@ -404,12 +420,12 @@ void FFractureEditorMode::GetActorGlobalBounds(TArrayView<UGeometryCollectionCom
 	}
 }
 
-void FFractureEditorMode::SelectionStateChanged()
+void UFractureEditorMode::SelectionStateChanged()
 {
 
 }
 
-void FFractureEditorMode::HandlePackageReloaded(const EPackageReloadPhase InPackageReloadPhase, FPackageReloadedEvent* InPackageReloadedEvent)
+void UFractureEditorMode::HandlePackageReloaded(const EPackageReloadPhase InPackageReloadPhase, FPackageReloadedEvent* InPackageReloadedEvent)
 {
 	if (InPackageReloadPhase == EPackageReloadPhase::PostPackageFixup)
 	{
@@ -438,7 +454,7 @@ void FFractureEditorMode::HandlePackageReloaded(const EPackageReloadPhase InPack
 	}
 }
 
-FConvexVolume FFractureEditorMode::TranformFrustum(const FConvexVolume& InFrustum, const FMatrix& InMatrix)
+FConvexVolume UFractureEditorMode::TranformFrustum(const FConvexVolume& InFrustum, const FMatrix& InMatrix)
 {
 	FConvexVolume NewFrustum;
 	NewFrustum.Planes.Empty(6);
@@ -453,7 +469,7 @@ FConvexVolume FFractureEditorMode::TranformFrustum(const FConvexVolume& InFrustu
 	return NewFrustum;
 }
 
-FConvexVolume FFractureEditorMode::GetVolumeFromBox(const FBox &InBox)
+FConvexVolume UFractureEditorMode::GetVolumeFromBox(const FBox &InBox)
 {
 	FConvexVolume ConvexVolume;
 	ConvexVolume.Planes.Empty(6);
diff --git a/Engine/Plugins/Experimental/ChaosEditor/Source/FractureEditor/Private/FractureEditorModeToolkit.cpp b/Engine/Plugins/Experimental/ChaosEditor/Source/FractureEditor/Private/FractureEditorModeToolkit.cpp
--- a/Engine/Plugins/Experimental/ChaosEditor/Source/FractureEditor/Private/FractureEditorModeToolkit.cpp
+++ b/Engine/Plugins/Experimental/ChaosEditor/Source/FractureEditor/Private/FractureEditorModeToolkit.cpp
@@ -183,6 +183,11 @@ FFractureEditorModeToolkit::~FFractureEditorModeToolkit()
 }
 
 void FFractureEditorModeToolkit::Init(const TSharedPtr<IToolkitHost>& InitToolkitHost)
+{
+	Init(InitToolkitHost, TWeakObjectPtr<UEdMode>());
+}
+
+void FFractureEditorModeToolkit::Init(const TSharedPtr<IToolkitHost>& InitToolkitHost, TWeakObjectPtr<UEdMode> InOwningMode)
 {
 	FFractureEditorModule& FractureModule = FModuleManager::GetModuleChecked<FFractureEditorModule>("FractureEditor");
 
@@ -411,7 +416,7 @@ void FFractureEditorModeToolkit::Init(const TSharedPtr<IToolkitHost>& InitToolki
 	// Bind Chaos Commands;
 	BindCommands();
 
-	FModeToolkit::Init(InitToolkitHost);
+	FModeToolkit::Init(InitToolkitHost, InOwningMode);
 
 }
 
@@ -626,7 +631,7 @@ FText FFractureEditorModeToolkit::GetBaseToolkitName() const
 
 class FEdMode* FFractureEditorModeToolkit::GetEditorMode() const
 {
-	return GLevelEditorModeTools().GetActiveMode(FFractureEditorMode::EM_FractureEditorModeId);
+	return GLevelEditorModeTools().GetActiveMode(UFractureEditorMode::EM_FractureEditorModeId);
 }
 
 void FFractureEditorModeToolkit::AddReferencedObjects(FReferenceCollector& Collector)
@@ -863,6 +868,11 @@ bool FFractureEditorModeToolkit::CanExecuteAction(UFractureActionTool* InActionT
 
 void FFractureEditorModeToolkit::SetActiveTool(UFractureModalTool* InActiveTool)
 {
+	if (ActiveTool)
+	{
+		ActiveTool->Shutdown();
+	}
+
 	ActiveTool = InActiveTool;
 
 	UFractureToolSettings* ToolSettings = GetMutableDefault<UFractureToolSettings>();
@@ -873,14 +883,25 @@ void FFractureEditorModeToolkit::SetActiveTool(UFractureModalTool* InActiveTool)
 
 	if (ActiveTool != nullptr)
 	{
+		ActiveTool->Setup();
+
 		Settings.Append(ActiveTool->GetSettingsObjects());
 
+		ActiveTool->SelectedBonesChanged();
 		ActiveTool->FractureContextChanged();
 	}
 
 	DetailsView->SetObjects(Settings);
 }
 
+void FFractureEditorModeToolkit::Shutdown()
+{
+	if (ActiveTool)
+	{
+		ActiveTool->Shutdown();
+	}
+}
+
 
 UFractureModalTool* FFractureEditorModeToolkit::GetActiveTool() const
 {
@@ -943,6 +964,7 @@ void FFractureEditorModeToolkit::SetOutlinerComponents(const TArray<UGeometryCol
 
 	if (ActiveTool != nullptr)
 	{
+		ActiveTool->SelectedBonesChanged();
 		ActiveTool->FractureContextChanged();
 	}
 }
@@ -953,6 +975,7 @@ void FFractureEditorModeToolkit::SetBoneSelection(UGeometryCollectionComponent*
 	
 	if (ActiveTool != nullptr)
 	{
+		ActiveTool->SelectedBonesChanged();
 		ActiveTool->FractureContextChanged();
 	}
 }
@@ -1335,6 +1358,7 @@ void FFractureEditorModeToolkit::OnOutlinerBoneSelectionChanged(UGeometryCollect
 
 		if (ActiveTool != nullptr)
 		{
+			ActiveTool->SelectedBonesChanged();
 			ActiveTool->FractureContextChanged();
 		}
 
diff --git a/Engine/Plugins/Experimental/ChaosEditor/Source/FractureEditor/Private/FractureToolCutter.cpp b/Engine/Plugins/Experimental/ChaosEditor/Source/FractureEditor/Private/FractureToolCutter.cpp
--- a/Engine/Plugins/Experimental/ChaosEditor/Source/FractureEditor/Private/FractureToolCutter.cpp
+++ b/Engine/Plugins/Experimental/ChaosEditor/Source/FractureEditor/Private/FractureToolCutter.cpp
@@ -3,6 +3,8 @@
 #include "FractureToolCutter.h"
 
 #include "FractureToolContext.h"
+#include "InteractiveToolsContext.h"
+#include "EditorModeManager.h"
 
 #include "GeometryCollection/GeometryCollectionObject.h"
 #include "GeometryCollection/GeometryCollectionComponent.h"
@@ -14,6 +16,86 @@
 #define LOCTEXT_NAMESPACE "FractureToolCutter"
 
 
+UFractureTransformGizmoSettings::UFractureTransformGizmoSettings(const FObjectInitializer& ObjInit) : Super(ObjInit)
+{
+
+}
+
+void UFractureTransformGizmoSettings::ResetGizmo(bool bResetRotation)
+{
+	if (!TransformGizmo || !TransformProxy)
+	{
+		return;
+	}
+	if (AttachedCutter)
+	{
+		AttachedCutter->SetMandateGroupFracture(bUseGizmo);
+	}
+	if (!bUseGizmo || !AttachedCutter)
+	{
+		TransformGizmo->SetVisibility(false);
+		return;
+	}
+	FBox CombinedBounds = AttachedCutter->GetCombinedBounds(AttachedCutter->GetFractureToolContexts());
+	TransformGizmo->SetVisibility((bool)CombinedBounds.IsValid);
+	if (CombinedBounds.IsValid)
+	{
+		if (bCenterOnSelection)
+		{
+			FTransform Transform = TransformProxy->GetTransform();
+			Transform.SetTranslation(CombinedBounds.GetCenter());
+			if (bResetRotation)
+			{
+				Transform.SetRotation(FQuat::Identity);
+			}
+			TransformGizmo->SetNewGizmoTransform(Transform);
+		}
+	}
+}
+
+void UFractureTransformGizmoSettings::PostEditChangeProperty(FPropertyChangedEvent& PropertyChangedEvent)
+{
+	if (PropertyChangedEvent.Property->GetFName() == GET_MEMBER_NAME_CHECKED(UFractureTransformGizmoSettings, bUseGizmo))
+	{
+		if (AttachedCutter)
+		{
+			ResetGizmo();
+		}
+	}
+}
+
+void UFractureTransformGizmoSettings::TransformChanged(UTransformProxy* Proxy, FTransform Transform)
+{
+	if (bUseGizmo && AttachedCutter)
+	{
+		AttachedCutter->FractureContextChanged();
+	}
+}
+
+void UFractureTransformGizmoSettings::Setup(UFractureToolCutterBase* Cutter)
+{
+	AttachedCutter = Cutter;
+	UInteractiveToolsContext* Context = GLevelEditorModeTools().GetInteractiveToolsContext();
+	if (ensure(Context && AttachedCutter))
+	{
+		UInteractiveGizmoManager* GizmoManager = Context->GizmoManager;
+		TransformProxy = NewObject<UTransformProxy>(this);
+		TransformGizmo = GizmoManager->CreateCustomTransformGizmo(ETransformGizmoSubElements::StandardTranslateRotate, this);
+		TransformGizmo->SetActiveTarget(TransformProxy);
+		TransformProxy->OnTransformChanged.AddUObject(this, &UFractureTransformGizmoSettings::TransformChanged);
+		ResetGizmo();
+	}
+}
+
+void UFractureTransformGizmoSettings::Shutdown()
+{
+	UInteractiveToolsContext* Context = GLevelEditorModeTools().GetInteractiveToolsContext();
+	if (Context)
+	{
+		Context->GizmoManager->DestroyAllGizmosByOwner(this);
+	}
+}
+
 UFractureToolCutterBase::UFractureToolCutterBase(const FObjectInitializer& ObjInit)
 	: Super(ObjInit)
 {
@@ -106,6 +188,17 @@ TArray<FFractureToolContext> UFractureToolCutterBase::GetFractureToolContexts()
 }
 
 
+FBox UFractureToolCutterBase::GetCombinedBounds(const TArray<FFractureToolContext>& Contexts) const
+{
+	FBox CombinedBounds(EForceInit::ForceInit);
+	for (const FFractureToolContext& FractureContext : Contexts)
+	{
+		CombinedBounds += FractureContext.GetWorldBounds();
+	}
+	return CombinedBounds;
+}
+
+
 void UFractureToolCutterBase::UpdateDefaultRandomSeed()
 {
 	DefaultRandomSeed = FMath::Rand();
diff --git a/Engine/Plugins/Experimental/ChaosEditor/Source/FractureEditor/Private/FractureToolCutter.h b/Engine/Plugins/Experimental/ChaosEditor/Source/FractureEditor/Private/FractureToolCutter.h
--- a/Engine/Plugins/Experimental/ChaosEditor/Source/FractureEditor/Private/FractureToolCutter.h
+++ b/Engine/Plugins/Experimental/ChaosEditor/Source/FractureEditor/Private/FractureToolCutter.h
@@ -4,6 +4,9 @@
 
 #include "FractureTool.h"
 
+#include "BaseGizmos/TransformGizmo.h"
+#include "BaseGizmos/TransformProxy.h"
+
 #include "FractureToolCutter.generated.h"
 
 class FFractureToolContext;
@@ -38,9 +41,13 @@ public:
 	float ChanceToFracture;
 
 	/** Generate a fracture pattern across all selected meshes.  */
-	UPROPERTY(EditAnywhere, Category = CommonFracture, meta = (DisplayName = "Group Fracture"))
+	UPROPERTY(EditAnywhere, Category = CommonFracture, meta = (EditCondition = "bGroupFractureToggleEnabled", HideEditConditionToggle, EditConditionHides, DisplayName = "Group Fracture"))
 	bool bGroupFracture;
 
+	// This flag allows tools to disable the above bGroupFracture option if/when it is not applicable
+	UPROPERTY()
+	bool bGroupFractureToggleEnabled = true;
+
 	/** Generate a fracture pattern across all selected meshes.  */
 	UPROPERTY(EditAnywhere, Category = CommonFracture, meta = (DisplayName = "Draw Sites"))
 	bool bDrawSites;
@@ -107,6 +114,17 @@ public:
 		UpdateDefaultRandomSeed();
 	}
 
+	FBox GetCombinedBounds(const TArray<FFractureToolContext>& Contexts) const;
+
+	void SetMandateGroupFracture(bool bMandateGroupFracture)
+	{
+		CutterSettings->bGroupFractureToggleEnabled = !bMandateGroupFracture;
+		if (bMandateGroupFracture)
+		{
+			CutterSettings->bGroupFracture = true;
+		}
+	}
+
 protected:
 	UPROPERTY(EditAnywhere, Category = Slicing)
 	TObjectPtr<UFractureCutterSettings> CutterSettings;
@@ -145,3 +163,53 @@ private:
 	TArray<FLinearColor> Colors;
 };
 
+
+/// This helps create a 3D transform gizmo that can be used to adjust fracture placement
+// Note it is tailored to UFractureToolCutterBase, and expects Setup(), Shutdown()
+// and ResetGizmo() to be called on tool setup, shutdown, and selection change respectively
+UCLASS(config = EditorPerProjectUserSettings)
+class UFractureTransformGizmoSettings : public UFractureToolSettings
+{
+public:
+	GENERATED_BODY()
+
+	UFractureTransformGizmoSettings(const FObjectInitializer & ObjInit);
+
+	void Setup(UFractureToolCutterBase* Cutter);
+	void Shutdown();
+
+	virtual void PostEditChangeProperty(FPropertyChangedEvent& PropertyChangedEvent) override;
+
+	void TransformChanged(UTransformProxy* Proxy, FTransform Transform);
+
+	void ResetGizmo(bool bResetRotation = false);
+
+	bool IsGizmoEnabled()
+	{
+		return bUseGizmo;
+	}
+
+	FTransform GetTransform()
+	{
+		return TransformProxy->GetTransform();
+	}
+
+	/** Use a 3D rigid transform gizmo to place the fracture pattern.  Only supports grouped fracture. */
+	UPROPERTY(EditAnywhere, Category = PlacementControls)
+	bool bUseGizmo = true;
+
+	/** Recenter the gizmo to the center of the selection when selection changes */
+	UPROPERTY(EditAnywhere, Category = PlacementControls, meta = (EditCondition = "bUseGizmo"))
+	bool bCenterOnSelection = true;
+
+	UPROPERTY()
+	TObjectPtr<UTransformGizmo> TransformGizmo = nullptr;
+
+	UPROPERTY()
+	TObjectPtr<UTransformProxy> TransformProxy = nullptr;
+
+protected:
+
+	TObjectPtr<UFractureToolCutterBase> AttachedCutter = nullptr;
+};
+
diff --git a/Engine/Plugins/Experimental/ChaosEditor/Source/FractureEditor/Private/FractureToolRadial.cpp b/Engine/Plugins/Experimental/ChaosEditor/Source/FractureEditor/Private/FractureToolRadial.cpp
--- a/Engine/Plugins/Experimental/ChaosEditor/Source/FractureEditor/Private/FractureToolRadial.cpp
+++ b/Engine/Plugins/Experimental/ChaosEditor/Source/FractureEditor/Private/FractureToolRadial.cpp
@@ -13,6 +13,19 @@ UFractureToolRadial::UFractureToolRadial(const FObjectInitializer& ObjInit)
 {
 	RadialSettings = NewObject<UFractureRadialSettings>(GetTransientPackage(), UFractureRadialSettings::StaticClass());
 	RadialSettings->OwnerTool = this;
+	GizmoSettings = NewObject<UFractureTransformGizmoSettings>(GetTransientPackage(), UFractureTransformGizmoSettings::StaticClass());
+	GizmoSettings->OwnerTool = this;
+}
+
+void UFractureToolRadial::Setup()
+{
+	GizmoSettings->Setup(this);
+}
+
+
+void UFractureToolRadial::Shutdown()
+{
+	GizmoSettings->Shutdown();
 }
 
 FText UFractureToolRadial::GetDisplayText() const
@@ -40,6 +53,7 @@ TArray<UObject*> UFractureToolRadial::GetSettingsObjects() const
 { 
 	TArray<UObject*> Settings; 
 	Settings.Add(CutterSettings);
+	Settings.Add(GizmoSettings);
 	Settings.Add(CollisionSettings);
 	Settings.Add(RadialSettings);
 	return Settings;
@@ -51,7 +65,7 @@ void UFractureToolRadial::GenerateVoronoiSites(const FFractureToolContext& Conte
 	const FVector::FReal AngularStep = 2 * PI / RadialSettings->AngularSteps;
 
 	FBox Bounds = Context.GetWorldBounds();
-	const FVector Center(Bounds.GetCenter() + RadialSettings->Center);
+	FVector Center(Bounds.GetCenter() + RadialSettings->Center);
 
 	FRandomStream RandStream(Context.GetSeed());
 	FVector UpVector(RadialSettings->Normal);
@@ -59,6 +73,15 @@ void UFractureToolRadial::GenerateVoronoiSites(const FFractureToolContext& Conte
 	FVector BasisX, BasisY;
 	UpVector.FindBestAxisVectors(BasisX, BasisY);
 
+	if (GizmoSettings->IsGizmoEnabled())
+	{
+		const FTransform& Transform = GizmoSettings->GetTransform();
+		UpVector = Transform.GetUnitAxis(EAxis::Z);
+		BasisX = Transform.GetUnitAxis(EAxis::X);
+		BasisY = Transform.GetUnitAxis(EAxis::Y);
+		Center = Transform.GetTranslation();
+	}
+
 	FVector::FReal Len = RadialStep * .5;
 	for (int32 ii = 0; ii < RadialSettings->RadialSteps; ++ii, Len += RadialStep)
 	{
@@ -71,4 +94,11 @@ void UFractureToolRadial::GenerateVoronoiSites(const FFractureToolContext& Conte
 	}
 }
 
+void UFractureToolRadial::SelectedBonesChanged()
+{
+	GizmoSettings->ResetGizmo();
+	Super::SelectedBonesChanged();
+}
+
+
 #undef LOCTEXT_NAMESPACE
diff --git a/Engine/Plugins/Experimental/ChaosEditor/Source/FractureEditor/Private/FractureToolRadial.h b/Engine/Plugins/Experimental/ChaosEditor/Source/FractureEditor/Private/FractureToolRadial.h
--- a/Engine/Plugins/Experimental/ChaosEditor/Source/FractureEditor/Private/FractureToolRadial.h
+++ b/Engine/Plugins/Experimental/ChaosEditor/Source/FractureEditor/Private/FractureToolRadial.h
@@ -69,6 +69,7 @@ public:
 	virtual FText GetTooltipText() const override;
 	virtual FSlateIcon GetToolIcon() const override;
 	virtual TArray<UObject*> GetSettingsObjects() const override;
+	virtual void SelectedBonesChanged() override;
 
 	virtual void RegisterUICommand( FFractureEditorCommands* BindingContext ) override;
 
@@ -76,6 +77,12 @@ public:
 	UPROPERTY(EditAnywhere, Category = Uniform)
 	TObjectPtr<UFractureRadialSettings> RadialSettings;
 
+	UPROPERTY(EditAnywhere, Category = Uniform)
+	TObjectPtr<UFractureTransformGizmoSettings> GizmoSettings;
+
+	virtual void Setup() override;
+	virtual void Shutdown() override;
+
 protected:
 	void GenerateVoronoiSites(const FFractureToolContext& Context, TArray<FVector>& Sites) override;
 
diff --git a/Engine/Plugins/Experimental/ChaosEditor/Source/FractureEditor/Public/FractureEditorMode.h b/Engine/Plugins/Experimental/ChaosEditor/Source/FractureEditor/Public/FractureEditorMode.h
--- a/Engine/Plugins/Experimental/ChaosEditor/Source/FractureEditor/Public/FractureEditorMode.h
+++ b/Engine/Plugins/Experimental/ChaosEditor/Source/FractureEditor/Public/FractureEditorMode.h
@@ -3,10 +3,12 @@
 #pragma once
 
 #include "CoreMinimal.h"
-#include "EdMode.h"
+#include "Tools/LegacyEdModeWidgetHelpers.h"
 #include "EditorUndoClient.h"
 #include "GeometryCollection/GeometryCollectionActor.h"
 
+#include "FractureEditorMode.generated.h"
+
 class UGeometryCollectionComponent;
 
 namespace FractureTransactionContexts
@@ -14,39 +16,44 @@ namespace FractureTransactionContexts
 	static const TCHAR SelectBoneContext[] = TEXT("SelectGeometryCollectionBone");
 };
 
-class FFractureEditorMode : public FEdMode, public FEditorUndoClient
+UCLASS(Transient)
+class UFractureEditorMode : public UBaseLegacyWidgetEdMode, public FEditorUndoClient, public ILegacyEdModeSelectInterface
 {
+	GENERATED_BODY()
 public:
 	const static FEditorModeID EM_FractureEditorModeId;
-public:
-	FFractureEditorMode();
-	virtual ~FFractureEditorMode();
 
-	using FEdMode::PostUndo;
+	UFractureEditorMode();
+	virtual ~UFractureEditorMode();
+
+	using UEdMode::PostUndo;
 	using FEditorUndoClient::PostUndo;
 
-	// FEdMode interface 
+	// UEdMode interface
 	virtual void Enter() override;
 	virtual void Exit() override;
-	virtual void AddReferencedObjects(FReferenceCollector& Collector) override;
-
-	// FEditorUndoClient interface
-	virtual bool MatchesContext(const FTransactionContext& InContext, const TArray<TPair<UObject *, FTransactionObjectEvent>>& TransactionObjectContexts) const override;
-	virtual void PostUndo(bool bSuccess) override; 
-	virtual void PostRedo(bool bSuccess) override;
 
 	//virtual void Tick(FEditorViewportClient* ViewportClient, float DeltaTime) override;
 	virtual void Render(const FSceneView* View, FViewport* Viewport, FPrimitiveDrawInterface* PDI) override;
 
+	virtual void CreateToolkit() override;
 	bool UsesToolkits() const override;
 	virtual bool InputKey(FEditorViewportClient* ViewportClient, FViewport* Viewport, FKey Key, EInputEvent Event) override;
 	virtual bool HandleClick(FEditorViewportClient* InViewportClient, HHitProxy* HitProxy, const FViewportClick& Click) override;
 
-	virtual bool BoxSelect(FBox& InBox, bool InSelect = true) override;
-	virtual bool FrustumSelect(const FConvexVolume& InFrustum, FEditorViewportClient* InViewportClient, bool InSelect) override;
 	virtual bool ComputeBoundingBoxForViewportFocus(AActor* Actor, UPrimitiveComponent* PrimitiveComponent, FBox& InOutBox) const;
 	virtual bool GetPivotForOrbit(FVector& OutPivot) const override;
-	// End of FEdMode interface
+	// End of UEdMode interface
+
+	// ILegacyEdModeSelectInterface
+	virtual bool BoxSelect(FBox& InBox, bool InSelect = true) override;
+	virtual bool FrustumSelect(const FConvexVolume& InFrustum, FEditorViewportClient* InViewportClient, bool InSelect) override;
+
+	// FEditorUndoClient interface
+	virtual bool MatchesContext(const FTransactionContext& InContext, const TArray<TPair<UObject*, FTransactionObjectEvent>>& TransactionObjectContexts) const override;
+	virtual void PostUndo(bool bSuccess) override;
+	virtual void PostRedo(bool bSuccess) override;
+
 private:
 	void OnUndoRedo();
 	void OnActorSelectionChanged(const TArray<UObject*>& NewSelection, bool bForceRefresh);
diff --git a/Engine/Plugins/Experimental/ChaosEditor/Source/FractureEditor/Public/FractureEditorModeToolkit.h b/Engine/Plugins/Experimental/ChaosEditor/Source/FractureEditor/Public/FractureEditorModeToolkit.h
--- a/Engine/Plugins/Experimental/ChaosEditor/Source/FractureEditor/Public/FractureEditorModeToolkit.h
+++ b/Engine/Plugins/Experimental/ChaosEditor/Source/FractureEditor/Public/FractureEditorModeToolkit.h
@@ -82,6 +82,7 @@ public:
 	
 	/** FModeToolkit interface */
 	virtual void Init(const TSharedPtr<IToolkitHost>& InitToolkitHost) override;
+	virtual void Init(const TSharedPtr<IToolkitHost>& InitToolkitHost, TWeakObjectPtr<UEdMode> InOwningMode) override;
 
 	/** IToolkit interface */
 	virtual FName GetToolkitFName() const override;
@@ -103,6 +104,8 @@ public:
 	UFractureModalTool* GetActiveTool() const;
 	bool IsActiveTool(UFractureModalTool* InActiveTool);
 
+	void Shutdown();
+
 	void SetOutlinerComponents(const TArray<UGeometryCollectionComponent*>& InNewComponents);
 	void SetBoneSelection(UGeometryCollectionComponent* InRootComponent, const TArray<int32>& InSelectedBones, bool bClearCurrentSelection);
 
diff --git a/Engine/Plugins/Experimental/ChaosEditor/Source/FractureEditor/Public/FractureTool.h b/Engine/Plugins/Experimental/ChaosEditor/Source/FractureEditor/Public/FractureTool.h
--- a/Engine/Plugins/Experimental/ChaosEditor/Source/FractureEditor/Public/FractureTool.h
+++ b/Engine/Plugins/Experimental/ChaosEditor/Source/FractureEditor/Public/FractureTool.h
@@ -10,6 +10,8 @@
 #include "FractureEditorCommands.h"
 #include "FractureEditorModeToolkit.h"
 
+#include "EdModeInteractiveToolsContext.h"
+
 #include "FractureTool.generated.h"
 
 class UGeometryCollection;
@@ -107,6 +109,15 @@ public:
 
 	virtual void PostEditChangeChainProperty(struct FPropertyChangedChainEvent& PropertyChangedEvent) override;
 	virtual void FractureContextChanged() {}
+	
+	// Called when the selection changes, and when the modal tool is entered
+	virtual void SelectedBonesChanged() {}
+
+	// Called when the modal tool is entered
+	virtual void Setup() {}
+
+	// Called when the modal tool is exited (on switching to a new modal tool or exiting the fracture editor mode)
+	virtual void Shutdown() {}
 
 };
 
