make sure proximity is generated before running convex tool

#rnx

diff --git a/Engine/Plugins/Experimental/ChaosEditor/Source/FractureEditor/Private/FractureToolConvex.cpp b/Engine/Plugins/Experimental/ChaosEditor/Source/FractureEditor/Private/FractureToolConvex.cpp
--- a/Engine/Plugins/Experimental/ChaosEditor/Source/FractureEditor/Private/FractureToolConvex.cpp
+++ b/Engine/Plugins/Experimental/ChaosEditor/Source/FractureEditor/Private/FractureToolConvex.cpp
@@ -152,6 +152,7 @@ int32 UFractureToolConvex::ExecuteFracture(const FFractureToolContext& FractureC
 	if (FractureContext.GetGeometryCollection().IsValid())
 	{
 		FGeometryCollection& Collection = *FractureContext.GetGeometryCollection();
+		GenerateProximityIfNecessary(&Collection);
 		FGeometryCollectionConvexUtility::CreateNonOverlappingConvexHullData(&Collection, ConvexSettings->FractionAllowRemove);
 	}
 
