Improve voronoi fracture tools by expanding voronoi bounds to fit all voronoi sites

#rnx

diff --git a/Engine/Plugins/Experimental/ChaosEditor/Source/FractureEditor/Private/FractureToolCutter.cpp b/Engine/Plugins/Experimental/ChaosEditor/Source/FractureEditor/Private/FractureToolCutter.cpp
--- a/Engine/Plugins/Experimental/ChaosEditor/Source/FractureEditor/Private/FractureToolCutter.cpp
+++ b/Engine/Plugins/Experimental/ChaosEditor/Source/FractureEditor/Private/FractureToolCutter.cpp
@@ -248,9 +248,10 @@ void UFractureToolVoronoiCutterBase::FractureContextChanged()
 	{
 		// Move the local bounds to the actor so we we'll draw in the correct location
 		GenerateVoronoiSites(FractureContext, VoronoiSites);
+		FBox VoronoiBounds = GetVoronoiBounds(FractureContext, VoronoiSites);
 		if (CutterSettings->bDrawDiagram)
 		{
-			GetVoronoiEdges(VoronoiSites, FractureContext.GetWorldBounds(), VoronoiEdges, CellMember);
+			GetVoronoiEdges(VoronoiSites, VoronoiBounds, VoronoiEdges, CellMember);
 		}
 	}
 }
@@ -261,10 +262,7 @@ int32 UFractureToolVoronoiCutterBase::ExecuteFracture(const FFractureToolContext
 	{
 		TArray<FVector> Sites;
 		GenerateVoronoiSites(FractureContext, Sites);
-		FBox VoronoiBounds = FractureContext.GetWorldBounds();
-			
-		// expand bounds to make sure noise-perturbed voronoi cells still contain the whole input mesh
-		VoronoiBounds = VoronoiBounds.ExpandBy(CutterSettings->Amplitude + CutterSettings->Grout);
+		FBox VoronoiBounds = GetVoronoiBounds(FractureContext, Sites);
 			
 		FVoronoiDiagram Voronoi(Sites, VoronoiBounds, .1f);
 
@@ -289,4 +287,12 @@ int32 UFractureToolVoronoiCutterBase::ExecuteFracture(const FFractureToolContext
 	return INDEX_NONE;
 }
 
+FBox UFractureToolVoronoiCutterBase::GetVoronoiBounds(const FFractureToolContext& FractureContext, const TArray<FVector>& Sites) const
+{
+	FBox VoronoiBounds(Sites);
+	VoronoiBounds += FractureContext.GetWorldBounds();
+	VoronoiBounds.ExpandBy(CutterSettings->Amplitude + CutterSettings->Grout + KINDA_SMALL_NUMBER);
+	return VoronoiBounds;
+}
+
 #undef LOCTEXT_NAMESPACE
diff --git a/Engine/Plugins/Experimental/ChaosEditor/Source/FractureEditor/Private/FractureToolCutter.h b/Engine/Plugins/Experimental/ChaosEditor/Source/FractureEditor/Private/FractureToolCutter.h
--- a/Engine/Plugins/Experimental/ChaosEditor/Source/FractureEditor/Private/FractureToolCutter.h
+++ b/Engine/Plugins/Experimental/ChaosEditor/Source/FractureEditor/Private/FractureToolCutter.h
@@ -156,6 +156,8 @@ public:
 protected:
 	virtual void GenerateVoronoiSites(const FFractureToolContext& Context, TArray<FVector>& Sites) {}
 
+	FBox GetVoronoiBounds(const FFractureToolContext& FractureContext, const TArray<FVector>& Sites) const;
+
 private:
 	TArray<int32> CellMember;
 	TArray<TTuple<FVector, FVector>> VoronoiEdges;
