Fix crash when generating Brick Voronoi Fracture for a large number of bricks.

#rnx

diff --git a/Engine/Plugins/Experimental/ChaosEditor/Source/FractureEditor/Private/FractureToolBrick.cpp b/Engine/Plugins/Experimental/ChaosEditor/Source/FractureEditor/Private/FractureToolBrick.cpp
--- a/Engine/Plugins/Experimental/ChaosEditor/Source/FractureEditor/Private/FractureToolBrick.cpp
+++ b/Engine/Plugins/Experimental/ChaosEditor/Source/FractureEditor/Private/FractureToolBrick.cpp
@@ -53,46 +53,149 @@ TArray<UObject*> UFractureToolBrick::GetSettingsObjects() const
 	return Settings;
 }
 
+namespace FractureToolBrickLocals
+{
+	// Calculate total number of bricks based on given dimensions and the extent of the object to be fractured.
+	// If the input is not valid or the result is too large, this functions returns -1.
+	// It is possible that we are dealing with incredibly large meshes and small brick dimensions. Doing the
+	// calculations in double and checking for NaNs will catch cases where for integer we would have to jump through
+	// multiple hoops to make sure we are not dealing with overflow. And since the limit for the number of bricks is
+	// comparably low, we do not need to worry about loss in precision for very large integers.
+	// For example, running this calculation with brick dimensions of 1 for the Sky Sphere will result in integer
+	// overflow.
+	static int64 CalculateNumBricks(const FVector& Dimensions, const FVector& Extents)
+	{
+		if (Dimensions.GetMin() <= 0 || Extents.GetMin() <= 0)
+		{
+			return -1;
+		}
+
+		const FVector NumBricksPerDim(ceil(Extents.X / Dimensions.X), ceil(Extents.Y / Dimensions.Y), ceil(Extents.Z / Dimensions.Z));
+		if (NumBricksPerDim.ContainsNaN())
+		{
+			return -1;
+		}
+		
+		const double NumBricks = NumBricksPerDim.X * NumBricksPerDim.Y * NumBricksPerDim.Z;
+		if (FMath::IsNaN(NumBricks))
+		{
+			return -1;
+		}
+
+		return static_cast<int64>(NumBricks);
+	}
+
+	static FVector GetBrickDimensions(const UFractureBrickSettings* BrickSettings, const FVector& Extents)
+	{
+		// Limit for the total number of bricks.
+		const int64 NumBricksLimit = 1 << 18;
+
+		FVector Dimensions(BrickSettings->BrickLength, BrickSettings->BrickDepth, BrickSettings->BrickHeight);
+		
+		// Early out if we have inputs we cannot deal with. If this call to CalculateNumBricks is fine then any other
+		// call to it will be fine, too, and we do not need to check for invalid results again.
+		int64 NumBricks = CalculateNumBricks(Dimensions, Extents);
+		if (NumBricks < 0)
+		{
+			return FVector::ZeroVector;
+		}
+
+		if (NumBricks > NumBricksLimit)
+		{
+			const int64 InputNumBricks = NumBricks;
+
+			// Determine dimensions safely within the brick limit by iteratively doubling the brick size.
+			FVector SafeDimensions = Dimensions;
+			int64 SafeNumBricks;
+			do
+			{
+				SafeDimensions *= 2;
+				SafeNumBricks = CalculateNumBricks(SafeDimensions, Extents);
+			} while (SafeNumBricks > NumBricksLimit);
+
+			// Maximize brick dimensions to fit within the brick limit via iterative interval halving.
+			const int32 IterationsMax = 10;
+			int32 Iterations = 0;
+			do
+			{
+				const FVector MidDimensions = (Dimensions + SafeDimensions) / 2;
+				const int64 MidNumBricks = CalculateNumBricks(MidDimensions, Extents);
+
+				if (MidNumBricks > NumBricksLimit)
+				{
+					Dimensions = MidDimensions;
+					NumBricks = MidNumBricks;
+				}
+				else
+				{
+					SafeDimensions = MidDimensions;
+					SafeNumBricks = MidNumBricks;
+				}
+			} while (++Iterations < IterationsMax);
+
+			Dimensions = SafeDimensions;
+			NumBricks = SafeNumBricks;
+
+			UE_LOG(LogFractureTool, Warning, TEXT("Brick Voronoi Fracture: Current brick dimensions of %f x %f x %f would result in %d bricks. "
+				"Reduced brick dimensions to %f x %f x %f resulting in %d bricks to stay within maximum number of %d bricks."),
+				BrickSettings->BrickLength, BrickSettings->BrickDepth, BrickSettings->BrickHeight, InputNumBricks,
+				Dimensions.X, Dimensions.Y, Dimensions.Z, NumBricks, NumBricksLimit);
+		}
+
+		return Dimensions;
+	}
+}
+
 void UFractureToolBrick::GenerateBrickTransforms(const FBox& Bounds)
 {
 	const FVector& Min = Bounds.Min;
 	const FVector& Max = Bounds.Max;
 	const FVector Extents(Bounds.Max - Bounds.Min);
 
-	const FQuat HeaderRotation(FVector::UpVector, 1.5708);
+	// Determine brick dimensions (length, depth, height) and make sure we do not exceed the limit for the number of bricks.
+	// If we would simply use the input dimensions, we are prone to running out of memory and/or exceeding the storage capabilities of TArray, and crash.
+	const FVector BrickDimensions = FractureToolBrickLocals::GetBrickDimensions(BrickSettings.Get(), Extents);
 
-	const float HalfHeight = BrickSettings->BrickHeight * 0.5f;
-	const float HalfDepth = BrickSettings->BrickDepth * 0.5f;
-	const float HalfLength = BrickSettings->BrickLength * 0.5f;
+	// Early out if we have inputs we cannot deal with.
+	if (BrickDimensions == FVector::ZeroVector)
+	{
+		return;
+	}
+
+	// Reserve correct amount of memory to avoid re-allocations.
+	BrickTransforms.Reserve(FractureToolBrickLocals::CalculateNumBricks(BrickDimensions, Extents));
+
+	const FVector BrickHalfDimensions(BrickDimensions * 0.5);
+	const FQuat HeaderRotation(FVector::UpVector, 1.5708);
 
 	if (BrickSettings->Bond == EFractureBrickBond::Stretcher)
 	{
 		bool OddY = false;
-		for (float yy = 0.f; yy <= Extents.Y; yy += BrickSettings->BrickDepth)
+		for (float yy = 0.f; yy <= Extents.Y; yy += BrickDimensions.Y)
 		{
 			bool Oddline = false;
-			for (float zz = HalfHeight; zz <= Extents.Z; zz += BrickSettings->BrickHeight)
+			for (float zz = BrickHalfDimensions.Z; zz <= Extents.Z; zz += BrickDimensions.Z)
 			{
-				for (float xx = 0.f; xx <= Extents.X; xx += BrickSettings->BrickLength)
+				for (float xx = 0.f; xx <= Extents.X; xx += BrickDimensions.X)
 				{
-					FVector BrickPosition(Min + FVector(Oddline ^ OddY ? xx : xx + HalfLength, yy, zz));
+					FVector BrickPosition(Min + FVector(Oddline ^ OddY ? xx : xx + BrickHalfDimensions.X, yy, zz));
 					BrickTransforms.Emplace(FTransform(BrickPosition));
 				}
 				Oddline = !Oddline;
 			}
-				OddY = !OddY;
+			OddY = !OddY;
 		}
 	}
 	else if (BrickSettings->Bond == EFractureBrickBond::Stack)
 	{
 		bool OddY = false;
-		for (float yy = 0.f; yy <= Extents.Y; yy += BrickSettings->BrickDepth)
+		for (float yy = 0.f; yy <= Extents.Y; yy += BrickDimensions.Y)
 		{
-			for (float zz = HalfHeight; zz <= Extents.Z; zz += BrickSettings->BrickHeight)
+			for (float zz = BrickHalfDimensions.Z; zz <= Extents.Z; zz += BrickDimensions.Z)
 			{
-				for (float xx = 0.f; xx <= Extents.X; xx += BrickSettings->BrickLength)
+				for (float xx = 0.f; xx <= Extents.X; xx += BrickDimensions.X)
 				{
-					FVector BrickPosition(Min + FVector(OddY ? xx : xx + HalfLength, yy, zz));
+					FVector BrickPosition(Min + FVector(OddY ? xx : xx + BrickHalfDimensions.X, yy, zz));
 					BrickTransforms.Emplace(FTransform(BrickPosition));
 				}
 			}
@@ -101,26 +204,26 @@ void UFractureToolBrick::GenerateBrickTransforms(const FBox& Bounds)
 	}
 	else if (BrickSettings->Bond == EFractureBrickBond::English)
 	{
-		float HalfLengthDepthDifference = HalfLength - HalfDepth - HalfDepth;
+		float HalfLengthDepthDifference = BrickHalfDimensions.X - BrickHalfDimensions.Y - BrickHalfDimensions.Y;
 		bool OddY = false;
-		for (float yy = 0.f; yy <= Extents.Y; yy += BrickSettings->BrickDepth)
+		for (float yy = 0.f; yy <= Extents.Y; yy += BrickDimensions.Y)
 		{
 			bool Oddline = false;
-			for (float zz = HalfHeight; zz <= Extents.Z; zz += BrickSettings->BrickHeight)
+			for (float zz = BrickHalfDimensions.Z; zz <= Extents.Z; zz += BrickDimensions.Z)
 			{
 				if (Oddline && !OddY) // header row
 				{
-					for (float xx = 0.f; xx <= Extents.X; xx += BrickSettings->BrickDepth)
+					for (float xx = 0.f; xx <= Extents.X; xx += BrickDimensions.Y)
 					{
-						FVector BrickPosition(Min + FVector(Oddline ^ OddY ? xx : xx + HalfDepth, yy + HalfDepth, zz));
+						FVector BrickPosition(Min + FVector(Oddline ^ OddY ? xx : xx + BrickHalfDimensions.Y, yy + BrickHalfDimensions.Y, zz));
 						BrickTransforms.Emplace(FTransform(HeaderRotation, BrickPosition));
 					}
 				}
 				else if(!Oddline) // stretchers
 				{
-					for (float xx = 0.f; xx <= Extents.X; xx += BrickSettings->BrickLength)
+					for (float xx = 0.f; xx <= Extents.X; xx += BrickDimensions.X)
 					{
-						FVector BrickPosition(Min + FVector(Oddline ^ OddY ? xx : xx + HalfLength, OddY ? yy + HalfLengthDepthDifference : yy - HalfLengthDepthDifference , zz));
+						FVector BrickPosition(Min + FVector(Oddline ^ OddY ? xx : xx + BrickHalfDimensions.X, OddY ? yy + HalfLengthDepthDifference : yy - HalfLengthDepthDifference , zz));
 						BrickTransforms.Emplace(FTransform(BrickPosition));
 					}
 				}
@@ -132,14 +235,14 @@ void UFractureToolBrick::GenerateBrickTransforms(const FBox& Bounds)
 	else if (BrickSettings->Bond == EFractureBrickBond::Header)
 	{
 		bool OddY = false;
-		for (float yy = 0.f; yy <= Extents.Y; yy += BrickSettings->BrickLength)
+		for (float yy = 0.f; yy <= Extents.Y; yy += BrickDimensions.X)
 		{
 			bool Oddline = false;
-			for (float zz = HalfHeight; zz <= Extents.Z; zz += BrickSettings->BrickHeight)
+			for (float zz = BrickHalfDimensions.Z; zz <= Extents.Z; zz += BrickDimensions.Z)
 			{
-				for (float xx = 0.f; xx <= Extents.X; xx += BrickSettings->BrickDepth)
+				for (float xx = 0.f; xx <= Extents.X; xx += BrickDimensions.Y)
 				{
-					FVector BrickPosition(Min + FVector(Oddline ^ OddY ? xx : xx + HalfDepth, yy, zz));
+					FVector BrickPosition(Min + FVector(Oddline ^ OddY ? xx : xx + BrickHalfDimensions.Y, yy, zz));
 					BrickTransforms.Emplace(FTransform(HeaderRotation, BrickPosition));
 				}
 				Oddline = !Oddline;
@@ -149,16 +252,16 @@ void UFractureToolBrick::GenerateBrickTransforms(const FBox& Bounds)
 	}
 	else if (BrickSettings->Bond == EFractureBrickBond::Flemish)
 	{
-		float HalfLengthDepthDifference = HalfLength - BrickSettings->BrickDepth  ;
+		float HalfLengthDepthDifference = BrickHalfDimensions.X - BrickDimensions.Y;
 		bool OddY = false;
 		int32 RowY = 0;
-		for (float yy = 0.f; yy <= Extents.Y; yy += BrickSettings->BrickDepth)
+		for (float yy = 0.f; yy <= Extents.Y; yy += BrickDimensions.Y)
 		{
 			bool OddZ = false;
-			for (float zz = HalfHeight; zz <= Extents.Z; zz += BrickSettings->BrickHeight)
+			for (float zz = BrickHalfDimensions.Z; zz <= Extents.Z; zz += BrickDimensions.Z)
 			{
 				bool OddX = OddZ;
-				for (float xx = 0.f; xx <= Extents.X; xx += HalfLength + HalfDepth)
+				for (float xx = 0.f; xx <= Extents.X; xx += BrickHalfDimensions.X + BrickHalfDimensions.Y)
 				{
 					FVector BrickPosition(Min + FVector(xx,yy,zz));
 					if (OddX)
@@ -175,7 +278,7 @@ void UFractureToolBrick::GenerateBrickTransforms(const FBox& Bounds)
 					}
 					else if (!OddY) // header
 					{
-						BrickTransforms.Emplace(FTransform(HeaderRotation, BrickPosition + FVector(0, HalfDepth, 0))); // header
+						BrickTransforms.Emplace(FTransform(HeaderRotation, BrickPosition + FVector(0, BrickHalfDimensions.Y, 0))); // header
 					}
 					OddX = !OddX;
 				}
@@ -186,8 +289,8 @@ void UFractureToolBrick::GenerateBrickTransforms(const FBox& Bounds)
 		}
 	}
 
-	FVector BrickMax(HalfLength, HalfDepth, HalfHeight);
-	FVector BrickMin(-BrickMax);
+	const FVector BrickMax(BrickHalfDimensions);
+	const FVector BrickMin(-BrickHalfDimensions);
 
 	for (const auto& Transform : BrickTransforms)
 	{
@@ -269,7 +372,6 @@ void UFractureToolBrick::Render(const FSceneView* View, FViewport* Viewport, FPr
 
 }
 
-
 void UFractureToolBrick::AddBoxEdges(const FVector& Min, const FVector& Max)
 {
 	Edges.Emplace(MakeTuple(Min, FVector(Min.X, Max.Y, Min.Z)));
@@ -293,22 +395,27 @@ int32 UFractureToolBrick::ExecuteFracture(const FFractureToolContext& FractureCo
 	if (FractureContext.IsValid())
 	{
  		BrickTransforms.Empty();
-		GenerateBrickTransforms(FractureContext.GetWorldBounds());
 
-		const FQuat HeaderRotation(FVector::UpVector, 1.5708);
+		const FBox Bounds = FractureContext.GetWorldBounds();
+		GenerateBrickTransforms(Bounds);
 
-		const float HalfHeight = BrickSettings->BrickHeight * 0.5f;
-		const float HalfDepth = BrickSettings->BrickDepth * 0.5f;
-		const float HalfLength = BrickSettings->BrickLength * 0.5f;
+		// Get the same brick dimensions that were used in GenerateBrickTransform.
+		// If we cannot deal with the input data then the brick dimensions will be zero, but we do not need to
+		// explicitly handle that since it will only affect some local variables. The BrickTransforms will be empty
+		// and there are no further side effects.
+		const FVector BrickDimensions = FractureToolBrickLocals::GetBrickDimensions(BrickSettings, Bounds.Max - Bounds.Min);
+		const FVector BrickHalfDimensions(BrickDimensions * 0.5);
+
+		const FQuat HeaderRotation(FVector::UpVector, 1.5708);
 
 		TArray<FBox> BricksToCut;
 
 		// space the bricks by the grout setting, constrained to not erase the bricks or have zero grout
 		// (currently zero grout bricks would break assumptions in the fracture)
-		float MinDim = FMath::Min3(HalfHeight, HalfDepth, HalfLength);
-		float HalfGrout = FMath::Clamp(0.5f * CutterSettings->Grout, MinDim * 0.02f, MinDim * 0.98f);
-		FVector HalfBrick(HalfLength - HalfGrout, HalfDepth - HalfGrout, HalfHeight - HalfGrout);
-		FBox BrickBox(-HalfBrick, HalfBrick);
+		const float MinDim = FMath::Min3(BrickHalfDimensions.X, BrickHalfDimensions.Y, BrickHalfDimensions.Z);
+		const float HalfGrout = FMath::Clamp(0.5f * CutterSettings->Grout, MinDim * 0.02f, MinDim * 0.98f);
+		const FVector HalfBrick(BrickHalfDimensions - HalfGrout);
+		const FBox BrickBox(-HalfBrick, HalfBrick);
 
 		for (const FTransform& Trans : BrickTransforms)
 		{
diff --git a/Engine/Plugins/Experimental/ChaosEditor/Source/FractureEditor/Private/FractureToolBrick.h b/Engine/Plugins/Experimental/ChaosEditor/Source/FractureEditor/Private/FractureToolBrick.h
--- a/Engine/Plugins/Experimental/ChaosEditor/Source/FractureEditor/Private/FractureToolBrick.h
+++ b/Engine/Plugins/Experimental/ChaosEditor/Source/FractureEditor/Private/FractureToolBrick.h
@@ -64,7 +64,7 @@ public:
 	UPROPERTY(EditAnywhere, Category = Brick)
 	EFractureBrickProjection Up;
 
-	/** Brick length */
+	/** Brick Length */
 	UPROPERTY(EditAnywhere, Category = Brick, meta = (UIMin = "0.1", UIMax = "500.0", ClampMin = "0.001"))
 	float BrickLength;
 
@@ -72,7 +72,7 @@ public:
 	UPROPERTY(EditAnywhere, Category = Brick, meta = (UIMin = "0.1", UIMax = "500.0", ClampMin = "0.001"))
 	float BrickHeight;
 
-	/** Brick Height */
+	/** Brick Depth */
 	UPROPERTY(EditAnywhere, Category = Brick, meta = (UIMin = "0.1", UIMax = "500.0", ClampMin = "0.001"))
 	float BrickDepth;
 };
