Chaos Fracture Editor
-- Corrected issue that causes crash when active Geometry Collection asset is reloaded.

diff --git a/Engine/Plugins/Experimental/ChaosEditor/Source/FractureEditor/Private/FractureEditorMode.cpp b/Engine/Plugins/Experimental/ChaosEditor/Source/FractureEditor/Private/FractureEditorMode.cpp
--- a/Engine/Plugins/Experimental/ChaosEditor/Source/FractureEditor/Private/FractureEditorMode.cpp
+++ b/Engine/Plugins/Experimental/ChaosEditor/Source/FractureEditor/Private/FractureEditorMode.cpp
@@ -430,10 +430,10 @@ void UFractureEditorMode::HandlePackageReloaded(const EPackageReloadPhase InPack
 	if (InPackageReloadPhase == EPackageReloadPhase::PostPackageFixup)
 	{
 		// assemble referenced RestCollections
-		TSet<const UGeometryCollection*> ReferencedRestCollections;
+		TMap<const UGeometryCollection*, UGeometryCollectionComponent*> ReferencedRestCollections;
 		for (UGeometryCollectionComponent* ExistingSelection : SelectedGeometryComponents)
 		{
-			ReferencedRestCollections.Add(ExistingSelection->GetRestCollection());
+			ReferencedRestCollections.Add(TPair<const UGeometryCollection*, UGeometryCollectionComponent*>(ExistingSelection->GetRestCollection(), ExistingSelection));
 		}
 		
 		// refresh outliner if reloaded package contains a referenced RestCollection
@@ -446,6 +446,7 @@ void UFractureEditorMode::HandlePackageReloaded(const EPackageReloadPhase InPack
 					if (Toolkit.IsValid())
 					{
 						FFractureEditorModeToolkit* FractureToolkit = (FFractureEditorModeToolkit*)Toolkit.Get();
+						FFractureSelectionTools::ClearSelectedBones(ReferencedRestCollections[NewObject]);
 						FractureToolkit->SetOutlinerComponents(SelectedGeometryComponents);
 					}
 				}
