FractureEditor: Reduce limit for number of bricks in Brick Voronoi Fracture to 8k.

diff --git a/Engine/Plugins/Experimental/ChaosEditor/Source/FractureEditor/Private/FractureToolBrick.cpp b/Engine/Plugins/Experimental/ChaosEditor/Source/FractureEditor/Private/FractureToolBrick.cpp
--- a/Engine/Plugins/Experimental/ChaosEditor/Source/FractureEditor/Private/FractureToolBrick.cpp
+++ b/Engine/Plugins/Experimental/ChaosEditor/Source/FractureEditor/Private/FractureToolBrick.cpp
@@ -88,7 +88,7 @@ namespace FractureToolBrickLocals
 	static FVector GetBrickDimensions(const UFractureBrickSettings* BrickSettings, const FVector& Extents)
 	{
 		// Limit for the total number of bricks.
-		const int64 NumBricksLimit = 1 << 18;
+		const int64 NumBricksLimit = 8192;
 
 		FVector Dimensions(BrickSettings->BrickLength, BrickSettings->BrickDepth, BrickSettings->BrickHeight);
 		
