fix bounds adjustment on fracture w/ noise/grout

#rnx

diff --git a/Engine/Plugins/Experimental/ChaosEditor/Source/FractureEditor/Private/FractureToolCutter.cpp b/Engine/Plugins/Experimental/ChaosEditor/Source/FractureEditor/Private/FractureToolCutter.cpp
--- a/Engine/Plugins/Experimental/ChaosEditor/Source/FractureEditor/Private/FractureToolCutter.cpp
+++ b/Engine/Plugins/Experimental/ChaosEditor/Source/FractureEditor/Private/FractureToolCutter.cpp
@@ -291,8 +291,7 @@ FBox UFractureToolVoronoiCutterBase::GetVoronoiBounds(const FFractureToolContext
 {
 	FBox VoronoiBounds(Sites);
 	VoronoiBounds += FractureContext.GetWorldBounds();
-	VoronoiBounds.ExpandBy(CutterSettings->Amplitude + CutterSettings->Grout + KINDA_SMALL_NUMBER);
-	return VoronoiBounds;
+	return VoronoiBounds.ExpandBy(CutterSettings->Amplitude + CutterSettings->Grout + KINDA_SMALL_NUMBER);
 }
 
 #undef LOCTEXT_NAMESPACE
