change the fracture AutoUV options: Remove the 'normal z' alpha attribute from the attributes we bake, and instead consistently bake a island present-or-not boolean mask to the alpha channel, with no gutter, so that artists can more easily do external processing on the uv islands and visualize the boundaries.

#rnx

diff --git a/Engine/Plugins/Experimental/ChaosEditor/Source/FractureEditor/Private/FractureToolUV.cpp b/Engine/Plugins/Experimental/ChaosEditor/Source/FractureEditor/Private/FractureToolUV.cpp
--- a/Engine/Plugins/Experimental/ChaosEditor/Source/FractureEditor/Private/FractureToolUV.cpp
+++ b/Engine/Plugins/Experimental/ChaosEditor/Source/FractureEditor/Private/FractureToolUV.cpp
@@ -290,7 +290,7 @@ int32 UFractureToolAutoUV::ExecuteFracture(const FFractureToolContext& FractureC
 				AutoUVSettings->bDistToOuter ? (int32)EBakeAttributes::DistanceToExternal : 0,
 				AutoUVSettings->bAmbientOcclusion ? (int32)EBakeAttributes::AmbientOcclusion : 0,
 				AutoUVSettings->bSmoothedCurvature ? (int32)EBakeAttributes::Curvature : 0,
-				AutoUVSettings->bZNormal ? (int32)EBakeAttributes::NormalZ : 0
+				0
 			);
 		}
 		else
@@ -312,6 +312,7 @@ int32 UFractureToolAutoUV::ExecuteFracture(const FFractureToolContext& FractureC
 		AttribSettings.Curvature_ThicknessFactor = AutoUVSettings->ThicknessFactor;
 		AttribSettings.Curvature_MaxValue = AutoUVSettings->MaxCurvature;
 		AttribSettings.bNormalZ_TakeAbs = AutoUVSettings->bUseAbsoluteValue;
+		AttribSettings.ClearGutterChannel = 3; // default clear the gutters for the alpha channel, so it shows more clearly the island boundaries
 		UE::PlanarCut::TextureInternalSurfaces(UVLayer, Collection, FMath::CeilToInt(AutoUVSettings->GutterSize), Attributes, AttribSettings, ImageBuilder,
 			AutoUVSettings->TargetMaterialIDs != ETargetMaterialIDs::SelectedIDs, 
 			AutoUVSettings->TargetMaterialIDs == ETargetMaterialIDs::OddIDs ? EmptyMaterialIDs: AutoUVSettings->MaterialIDs);
diff --git a/Engine/Plugins/Experimental/ChaosEditor/Source/FractureEditor/Private/FractureToolUV.h b/Engine/Plugins/Experimental/ChaosEditor/Source/FractureEditor/Private/FractureToolUV.h
--- a/Engine/Plugins/Experimental/ChaosEditor/Source/FractureEditor/Private/FractureToolUV.h
+++ b/Engine/Plugins/Experimental/ChaosEditor/Source/FractureEditor/Private/FractureToolUV.h
@@ -76,8 +76,8 @@ public:
 	UPROPERTY(EditAnywhere, Category = MapSettings)
 	EAutoUVTextureResolution Resolution = EAutoUVTextureResolution::Resolution512;
 
-	/** Space to leave between UV islands, measured in texels */
-	UPROPERTY(EditAnywhere, Category = MapSettings, meta = (ClampMin = "1", ClampMax = "10", UIMax = "4"))
+	/** Approximate space to leave between UV islands, measured in texels */
+	UPROPERTY(EditAnywhere, Category = MapSettings, meta = (ClampMin = "1", ClampMax = "10", UIMax = "6"))
 	int32 GutterSize = 2;
 
 	/** The resulting automatically-generated texture map */
@@ -119,10 +119,6 @@ public:
 	UPROPERTY(EditAnywhere, Category = AttributesToBake, meta = (EditCondition = "BakeTextureType == ETextureType::ThicknessAndSurfaceAttributes", EditConditionHides))
 	bool bSmoothedCurvature = true;
 
-	/** Bake the Z-component of the normal to a texture channel (alpha) */
-	UPROPERTY(EditAnywhere, Category = AttributesToBake, meta = (EditCondition = "BakeTextureType == ETextureType::ThicknessAndSurfaceAttributes", EditConditionHides))
-	bool bZNormal = true;
-
 	/** Max distance to search for the outer mesh surface */
 	UPROPERTY(EditAnywhere, Category = DistToOuterSettings, meta = (EditCondition = "BakeTextureType == ETextureType::ThicknessAndSurfaceAttributes && bDistToOuter", EditConditionHides, UIMin = "1", UIMax = "100", ClampMin = ".01", ClampMax = "1000"))
 	double MaxDistance = 100;
diff --git a/Engine/Plugins/Experimental/PlanarCutPlugin/Source/PlanarCut/Private/FractureAutoUV.cpp b/Engine/Plugins/Experimental/PlanarCutPlugin/Source/PlanarCut/Private/FractureAutoUV.cpp
--- a/Engine/Plugins/Experimental/PlanarCutPlugin/Source/PlanarCut/Private/FractureAutoUV.cpp
+++ b/Engine/Plugins/Experimental/PlanarCutPlugin/Source/PlanarCut/Private/FractureAutoUV.cpp
@@ -565,7 +565,6 @@ void TextureInternalSurfaces(
 			FDynamicMeshUVOverlay* UV = Mesh.Attributes()->PrimaryUV();
 			for (int TID : Mesh.TriangleIndicesItr())
 			{
-				// TODO: separate the concept of inside from the concept of needs texture!!
 				bool bIsTextureTri = IsTriActive(
 					AugmentedDynamicMesh::GetVisibility(Mesh, TID), 
 					Mesh.Attributes()->GetMaterialID()->GetValue(TID), TargetMaterials, true, bOnlyOddMaterials);
@@ -784,9 +783,21 @@ void TextureInternalSurfaces(
 	CopyValuesToChannel(AmbientIdx, AmbientValues);
 	CopyValuesToChannel(CurvatureIdx, CurvatureValues);
 
-	for (TTuple<int64, int64>& GutterToInside : OccupancyMap.GutterTexels)
+	if (AttributeSettings.ClearGutterChannel > -1 && AttributeSettings.ClearGutterChannel < 4)
 	{
-		TextureOut.CopyPixel(GutterToInside.Value, GutterToInside.Key);
+		for (TTuple<int64, int64>& GutterToInside : OccupancyMap.GutterTexels)
+		{
+			FVector4f Pixel = TextureOut.GetPixel(GutterToInside.Value);
+			Pixel[AttributeSettings.ClearGutterChannel] = 0.0f;
+			TextureOut.SetPixel(GutterToInside.Key, Pixel);
+		}
+	}
+	else
+	{
+		for (TTuple<int64, int64>& GutterToInside : OccupancyMap.GutterTexels)
+		{
+			TextureOut.CopyPixel(GutterToInside.Value, GutterToInside.Key);
+		}
 	}
 }
 
diff --git a/Engine/Plugins/Experimental/PlanarCutPlugin/Source/PlanarCut/Public/FractureAutoUV.h b/Engine/Plugins/Experimental/PlanarCutPlugin/Source/PlanarCut/Public/FractureAutoUV.h
--- a/Engine/Plugins/Experimental/PlanarCutPlugin/Source/PlanarCut/Public/FractureAutoUV.h
+++ b/Engine/Plugins/Experimental/PlanarCutPlugin/Source/PlanarCut/Public/FractureAutoUV.h
@@ -59,6 +59,7 @@ struct FTextureAttributeSettings
 	double Curvature_BlurRadius = 2.5;
 	double Curvature_ThicknessFactor = 3.0; // distance to search for mesh correspondence, as a factor of voxel size
 	double Curvature_MaxValue = .1; // curvatures above this value will be clamped
+	int ClearGutterChannel = -1; // don't copy gutter values for this channel, if specified -- useful for visualizing the UV island borders
 	
 	bool bNormalZ_TakeAbs = true;
 };
