FractureTool
-- Added Flush Embedded Geometry tool

diff --git a/Engine/Plugins/Experimental/ChaosEditor/Source/FractureEditor/Private/FractureEditorCommands.h b/Engine/Plugins/Experimental/ChaosEditor/Source/FractureEditor/Private/FractureEditorCommands.h
--- a/Engine/Plugins/Experimental/ChaosEditor/Source/FractureEditor/Private/FractureEditorCommands.h
+++ b/Engine/Plugins/Experimental/ChaosEditor/Source/FractureEditor/Private/FractureEditorCommands.h
@@ -45,6 +45,7 @@ class FFractureEditorCommands : public TCommands<FFractureEditorCommands>
 		// Embed Commands
 		TSharedPtr< FUICommandInfo > AddEmbeddedGeometry;
 		TSharedPtr< FUICommandInfo > AutoEmbedGeometry;
+		TSharedPtr< FUICommandInfo > FlushEmbeddedGeometry;
 
 		// UV Commands
 		TSharedPtr< FUICommandInfo > AutoUV;
diff --git a/Engine/Plugins/Experimental/ChaosEditor/Source/FractureEditor/Private/FractureEditorModeToolkit.cpp b/Engine/Plugins/Experimental/ChaosEditor/Source/FractureEditor/Private/FractureEditorModeToolkit.cpp
--- a/Engine/Plugins/Experimental/ChaosEditor/Source/FractureEditor/Private/FractureEditorModeToolkit.cpp
+++ b/Engine/Plugins/Experimental/ChaosEditor/Source/FractureEditor/Private/FractureEditorModeToolkit.cpp
@@ -522,6 +522,7 @@ void FFractureEditorModeToolkit::BuildToolPalette(FName PaletteIndex, class FToo
 	{
 		ToolbarBuilder.AddToolBarButton(Commands.AddEmbeddedGeometry);
 		ToolbarBuilder.AddToolBarButton(Commands.AutoEmbedGeometry);
+		ToolbarBuilder.AddToolBarButton(Commands.FlushEmbeddedGeometry);
 	}
 	else if (PaletteIndex == TEXT("UV"))
 	{
diff --git a/Engine/Plugins/Experimental/ChaosEditor/Source/FractureEditor/Private/FractureEditorStyle.cpp b/Engine/Plugins/Experimental/ChaosEditor/Source/FractureEditor/Private/FractureEditorStyle.cpp
--- a/Engine/Plugins/Experimental/ChaosEditor/Source/FractureEditor/Private/FractureEditorStyle.cpp
+++ b/Engine/Plugins/Experimental/ChaosEditor/Source/FractureEditor/Private/FractureEditorStyle.cpp
@@ -49,6 +49,7 @@ FFractureEditorStyle::FFractureEditorStyle()
 
 	Set("FractureEditor.AddEmbeddedGeometry",  new IMAGE_BRUSH_SVG("FractureEmbed", IconSize));
 	Set("FractureEditor.AutoEmbedGeometry",    new IMAGE_BRUSH_SVG("FractureAutoEmbed", IconSize));
+	Set("FractureEditor.FlushEmbedGeometry",   new IMAGE_BRUSH_SVG("FractureFlushEmbed", IconSize));
 
 	Set("FractureEditor.AutoUV",               new IMAGE_BRUSH_SVG("FractureAutoUV", IconSize));
 
diff --git a/Engine/Plugins/Experimental/ChaosEditor/Source/FractureEditor/Private/FractureToolEmbed.cpp b/Engine/Plugins/Experimental/ChaosEditor/Source/FractureEditor/Private/FractureToolEmbed.cpp
--- a/Engine/Plugins/Experimental/ChaosEditor/Source/FractureEditor/Private/FractureToolEmbed.cpp
+++ b/Engine/Plugins/Experimental/ChaosEditor/Source/FractureEditor/Private/FractureToolEmbed.cpp
@@ -6,6 +6,8 @@
 #include "ScopedTransaction.h"
 #include "Engine/Selection.h"
 #include "Engine/StaticMesh.h"
+#include "ActorFactories/ActorFactoryStaticMesh.h"
+#include "AssetSelection.h"
 
 #include "GeometryCollection/GeometryCollectionComponent.h"
 #include "GeometryCollection/GeometryCollection.h"
@@ -310,4 +312,137 @@ TArray<UStaticMeshComponent*> UFractureToolAutoEmbedGeometry::GetSelectedStaticM
 }
 
 
+
+
+
+
+UFractureToolFlushEmbeddedGeometry::UFractureToolFlushEmbeddedGeometry(const FObjectInitializer& ObjInit)
+	: Super(ObjInit)
+{
+	FlushEmbeddedGeometrySettings = NewObject<UFractureToolFlushEmbeddedGeometrySettings>(GetTransientPackage(), UFractureToolFlushEmbeddedGeometrySettings::StaticClass());
+	FlushEmbeddedGeometrySettings->OwnerTool = this;
+}
+
+
+FText UFractureToolFlushEmbeddedGeometry::GetDisplayText() const
+{
+	return FText(NSLOCTEXT("Fracture", "FractureToolFlushEmbeddedGeometry", "Flush"));
+}
+
+
+FText UFractureToolFlushEmbeddedGeometry::GetTooltipText() const
+{
+	return FText(NSLOCTEXT("Fracture", "FractureToolCFlushEmbeddedGeometryToolTip", "Strip all instanced embedded geometry from the geometry collection."));
+}
+
+FText UFractureToolFlushEmbeddedGeometry::GetApplyText() const
+{
+	return FText(NSLOCTEXT("Fracture", "ExecuteFlushEmbeddedGeometry", "Flush Embedded Geometry"));
+}
+
+FSlateIcon UFractureToolFlushEmbeddedGeometry::GetToolIcon() const
+{
+	return FSlateIcon("FractureEditorStyle", "FractureEditor.FlushEmbeddedGeometry");
+}
+
+
+TArray<UObject*> UFractureToolFlushEmbeddedGeometry::GetSettingsObjects() const
+{
+	TArray<UObject*> Settings;
+	Settings.Add(FlushEmbeddedGeometrySettings);
+	return Settings;
+}
+
+void UFractureToolFlushEmbeddedGeometry::RegisterUICommand(FFractureEditorCommands* BindingContext)
+{
+	UI_COMMAND_EXT(BindingContext, UICommandInfo, "FlushEmbeddedGeometry", "Flush", "Strip all instanced embedded geometry from the geometry collection.", EUserInterfaceActionType::ToggleButton, FInputChord());
+	BindingContext->FlushEmbeddedGeometry = UICommandInfo;
+}
+
+void UFractureToolFlushEmbeddedGeometry::Execute(TWeakPtr<FFractureEditorModeToolkit> InToolkit)
+{
+	if (InToolkit.IsValid())
+	{
+		FFractureEditorModeToolkit* Toolkit = InToolkit.Pin().Get();
+
+		TArray<FFractureToolContext> Contexts = GetFractureToolContexts();
+
+		for (FFractureToolContext& Context : Contexts)
+		{
+			if (!Context.GetGeometryCollection()->HasAttribute("ExemplarIndex", FGeometryCollection::TransformGroup))
+			{
+				continue;
+			}
+			
+			const TManagedArray<int32>& ExemplarIndex = Context.GetGeometryCollection()->GetAttribute<int32>("ExemplarIndex", FGeometryCollection::TransformGroup);
+			const TManagedArray<FTransform>& Transforms = Context.GetGeometryCollection()->GetAttribute<FTransform>("Transform", FGeometryCollection::TransformGroup);
+			const TManagedArray<int32>& ParentIndices = Context.GetGeometryCollection()->GetAttribute<int32>("Parent", FGeometryCollection::TransformGroup);
+
+			TArray<FMatrix> GlobalMatrices;
+			GeometryCollectionAlgo::GlobalMatrices(Transforms, ParentIndices, GlobalMatrices);
+			
+			TArray<int32> RemovalList;
+			RemovalList.Reserve(ExemplarIndex.Num());
+
+			TMap<int32,TArray<FTransform>> ExemplarTransforms;
+
+			// Filter for instanced embedded geometry
+			for (int32 TransformIndex = 0; TransformIndex < ExemplarIndex.Num(); ++TransformIndex)
+			{
+				if (ExemplarIndex[TransformIndex] > INDEX_NONE)
+				{
+					RemovalList.Add(TransformIndex);
+					TArray<FTransform>& InstanceTransforms = ExemplarTransforms.FindOrAdd(ExemplarIndex[TransformIndex]);
+					InstanceTransforms.Add(FTransform(GlobalMatrices[TransformIndex]));
+				}
+			}
+
+			if (RemovalList.Num())
+			{
+				// Remove all instanced embedded geometry.
+				FGeometryCollectionEdit GeometryCollectionEdit = Context.GetGeometryCollectionComponent()->EditRestCollection(GeometryCollection::EEditUpdate::RestPhysicsDynamic);
+				UGeometryCollection* FracturedCollection = GeometryCollectionEdit.GetRestCollection();
+				FracturedCollection->GetGeometryCollection()->RemoveElements(FGeometryCollection::TransformGroup, RemovalList);
+
+				if (FlushEmbeddedGeometrySettings->bExtractAsStaticMeshActors)
+				{
+					FTransform ComponentSpace = Context.GetGeometryCollectionComponent()->GetComponentTransform();
+					
+					if (UWorld* TargetWorld = Context.GetGeometryCollectionComponent()->GetOwner()->GetWorld())
+					{
+						// Place Static Mesh Actors
+						for (auto It = ExemplarTransforms.CreateConstIterator(); It; ++It)
+						{
+							FSoftObjectPath ExemplarPath = FracturedCollection->EmbeddedGeometryExemplar[It->Key].StaticMeshExemplar;
+							UObject* LoadedExemplarAsset = ExemplarPath.TryLoad();
+							if (UActorFactory* ActorFactory = FActorFactoryAssetProxy::GetFactoryForAssetObject(LoadedExemplarAsset))
+							{
+
+								for (const FTransform& Transform : It->Value)
+								{
+									ActorFactory->CreateActor(LoadedExemplarAsset, TargetWorld->GetCurrentLevel(), ComponentSpace*Transform, RF_Transactional);
+								}
+							}
+
+						}
+					}
+					
+				}
+								
+				// Remove all Exemplars from asset.
+				FracturedCollection->EmbeddedGeometryExemplar.Empty();
+				
+				Context.GetGeometryCollectionComponent()->InitializeEmbeddedGeometry();
+
+				FScopedColorEdit ScopedEdit(Context.GetGeometryCollectionComponent());
+				ScopedEdit.ResetBoneSelection();
+				Refresh(Context, Toolkit);
+				FracturedCollection->MarkPackageDirty();
+			}	
+		}
+
+		SetOutlinerComponents(Contexts, Toolkit);
+	}
+}
+
 #undef LOCTEXT_NAMESPACE
\ No newline at end of file
diff --git a/Engine/Plugins/Experimental/ChaosEditor/Source/FractureEditor/Private/FractureToolEmbed.h b/Engine/Plugins/Experimental/ChaosEditor/Source/FractureEditor/Private/FractureToolEmbed.h
--- a/Engine/Plugins/Experimental/ChaosEditor/Source/FractureEditor/Private/FractureToolEmbed.h
+++ b/Engine/Plugins/Experimental/ChaosEditor/Source/FractureEditor/Private/FractureToolEmbed.h
@@ -47,3 +47,45 @@ public:
 private:
 	static TArray<UStaticMeshComponent*>  GetSelectedStaticMeshComponents();
 };
+
+
+UCLASS(DisplayName = "Flush Embedded Geometry", Category = "FractureTools")
+class UFractureToolFlushEmbeddedGeometrySettings : public UFractureToolSettings
+{
+public:
+	GENERATED_BODY()
+
+	UFractureToolFlushEmbeddedGeometrySettings(const FObjectInitializer& ObjInit)
+	: Super(ObjInit)
+	, bExtractAsStaticMeshActors(false)
+	{}
+
+	UPROPERTY(EditAnywhere, Category = EmbeddedGeometry, meta = (DisplayName = "Extract As StaticMesh Actors"))
+	bool bExtractAsStaticMeshActors;
+};
+
+UCLASS(DisplayName = "FlushEmbeddedGeometry", Category = "FractureTools")
+class UFractureToolFlushEmbeddedGeometry : public UFractureModalTool
+{
+public:
+	GENERATED_BODY()
+
+	using FGeometryCollectionPtr = TSharedPtr<FGeometryCollection, ESPMode::ThreadSafe>;
+
+	UFractureToolFlushEmbeddedGeometry(const FObjectInitializer& ObjInit);
+
+	// UFractureTool Interface
+	virtual FText GetDisplayText() const override;
+	virtual FText GetTooltipText() const override;
+	virtual FText GetApplyText() const override;
+	virtual FSlateIcon GetToolIcon() const override;
+	virtual TArray<UObject*> GetSettingsObjects() const override;
+	virtual void RegisterUICommand(FFractureEditorCommands* BindingContext);
+
+	virtual void Execute(TWeakPtr<FFractureEditorModeToolkit> InToolkit) override;
+
+	UPROPERTY(EditAnywhere, Category = EmbeddedGeometry)
+	TObjectPtr<UFractureToolFlushEmbeddedGeometrySettings> FlushEmbeddedGeometrySettings;
+
+
+};
