FractureTool
-- Updated FlushEmbeddedGeometry tool to use non-deprecated function for spawning actors.

diff --git a/Engine/Plugins/Experimental/ChaosEditor/Source/FractureEditor/Private/FractureToolEmbed.cpp b/Engine/Plugins/Experimental/ChaosEditor/Source/FractureEditor/Private/FractureToolEmbed.cpp
--- a/Engine/Plugins/Experimental/ChaosEditor/Source/FractureEditor/Private/FractureToolEmbed.cpp
+++ b/Engine/Plugins/Experimental/ChaosEditor/Source/FractureEditor/Private/FractureToolEmbed.cpp
@@ -406,6 +406,8 @@ void UFractureToolFlushEmbeddedGeometry::Execute(TWeakPtr<FFractureEditorModeToo
 
 				if (FlushEmbeddedGeometrySettings->bExtractAsStaticMeshActors)
 				{
+					FActorSpawnParameters SpawnParameters;
+					
 					FTransform ComponentSpace = Context.GetGeometryCollectionComponent()->GetComponentTransform();
 					
 					if (UWorld* TargetWorld = Context.GetGeometryCollectionComponent()->GetOwner()->GetWorld())
@@ -420,7 +422,7 @@ void UFractureToolFlushEmbeddedGeometry::Execute(TWeakPtr<FFractureEditorModeToo
 
 								for (const FTransform& Transform : It->Value)
 								{
-									ActorFactory->CreateActor(LoadedExemplarAsset, TargetWorld->GetCurrentLevel(), ComponentSpace*Transform, RF_Transactional);
+									ActorFactory->CreateActor(LoadedExemplarAsset, TargetWorld->GetCurrentLevel(), ComponentSpace*Transform, SpawnParameters);
 								}
 							}
 
