remove now-unused property from fracture autouv

#rnx

diff --git a/Engine/Plugins/Experimental/ChaosEditor/Source/FractureEditor/Private/FractureToolUV.cpp b/Engine/Plugins/Experimental/ChaosEditor/Source/FractureEditor/Private/FractureToolUV.cpp
--- a/Engine/Plugins/Experimental/ChaosEditor/Source/FractureEditor/Private/FractureToolUV.cpp
+++ b/Engine/Plugins/Experimental/ChaosEditor/Source/FractureEditor/Private/FractureToolUV.cpp
@@ -311,7 +311,6 @@ int32 UFractureToolAutoUV::ExecuteFracture(const FFractureToolContext& FractureC
 		AttribSettings.Curvature_VoxelRes = AutoUVSettings->VoxelResolution;
 		AttribSettings.Curvature_ThicknessFactor = AutoUVSettings->ThicknessFactor;
 		AttribSettings.Curvature_MaxValue = AutoUVSettings->MaxCurvature;
-		AttribSettings.bNormalZ_TakeAbs = AutoUVSettings->bUseAbsoluteValue;
 		AttribSettings.ClearGutterChannel = 3; // default clear the gutters for the alpha channel, so it shows more clearly the island boundaries
 		UE::PlanarCut::TextureInternalSurfaces(UVLayer, Collection, FMath::CeilToInt(AutoUVSettings->GutterSize), Attributes, AttribSettings, ImageBuilder,
 			AutoUVSettings->TargetMaterialIDs != ETargetMaterialIDs::SelectedIDs, 
diff --git a/Engine/Plugins/Experimental/ChaosEditor/Source/FractureEditor/Private/FractureToolUV.h b/Engine/Plugins/Experimental/ChaosEditor/Source/FractureEditor/Private/FractureToolUV.h
--- a/Engine/Plugins/Experimental/ChaosEditor/Source/FractureEditor/Private/FractureToolUV.h
+++ b/Engine/Plugins/Experimental/ChaosEditor/Source/FractureEditor/Private/FractureToolUV.h
@@ -151,10 +151,6 @@ public:
 	UPROPERTY(EditAnywhere, Category = SmoothedCurvatureSettings, meta = (EditCondition = "BakeTextureType == ETextureType::ThicknessAndSurfaceAttributes && bSmoothedCurvature", EditConditionHides, UIMin = ".01", UIMax = "1", ClampMin = ".0001", ClampMax = "10"))
 	double MaxCurvature = .1;
 
-	/** Whether to use the absolute value of the Z-Component of the normal */
-	UPROPERTY(EditAnywhere, Category = ZNormalSettings, meta = (EditCondition = "BakeTextureType == ETextureType::ThicknessAndSurfaceAttributes && bZNormal", EditConditionHides))
-	bool bUseAbsoluteValue = true;
-
 };
 
 
