FractureTool
-- force regeneration of guids for outliner items

diff --git a/Engine/Plugins/Experimental/ChaosEditor/Source/FractureEditor/Private/SGeometryCollectionOutliner.cpp b/Engine/Plugins/Experimental/ChaosEditor/Source/FractureEditor/Private/SGeometryCollectionOutliner.cpp
--- a/Engine/Plugins/Experimental/ChaosEditor/Source/FractureEditor/Private/SGeometryCollectionOutliner.cpp
+++ b/Engine/Plugins/Experimental/ChaosEditor/Source/FractureEditor/Private/SGeometryCollectionOutliner.cpp
@@ -492,7 +492,7 @@ void FGeometryCollectionTreeItemComponent::RegenerateChildren()
 		{
 			int32 NumElements = Collection->NumElements(FGeometryCollection::TransformGroup);
 
-			GeometryCollection::GenerateTemporaryGuids(Collection);
+			GeometryCollection::GenerateTemporaryGuids(Collection, 0, true);
 			const TManagedArray<FGuid>& Guids = Collection->GetAttribute<FGuid>("GUID", "Transform");
 			const TManagedArray<int32>& Parents = Collection->Parent;
 			// const TManagedArray<FString>& BoneNames = CollectionPtr->BoneName;
