make fracture autouv overwrite existing textures without deleting the existing files

#rnx

diff --git a/Engine/Plugins/Experimental/ChaosEditor/Source/FractureEditor/Private/FractureToolUV.cpp b/Engine/Plugins/Experimental/ChaosEditor/Source/FractureEditor/Private/FractureToolUV.cpp
--- a/Engine/Plugins/Experimental/ChaosEditor/Source/FractureEditor/Private/FractureToolUV.cpp
+++ b/Engine/Plugins/Experimental/ChaosEditor/Source/FractureEditor/Private/FractureToolUV.cpp
@@ -161,12 +161,9 @@ TArray<FFractureToolContext> UFractureToolAutoUV::GetFractureToolContexts() cons
 }
 
 
-bool UFractureToolAutoUV::SaveGeneratedTexture(UTexture2D* GeneratedTexture, FString ObjectBaseName, const UObject* RelativeToAsset, bool bPromptToSave, bool bAllowReplace)
+bool UFractureToolAutoUV::SaveGeneratedTexture(UE::Geometry::TImageBuilder<FVector4f>& ImageBuilder, FString ObjectBaseName, const UObject* RelativeToAsset, bool bPromptToSave, bool bAllowReplace)
 {
 	check(RelativeToAsset);
-	check(GeneratedTexture);
-	check(GeneratedTexture->GetOuter() == GetTransientPackage());
-	check(GeneratedTexture->Source.IsValid());	// texture needs to have valid source data to be savd
 
 	// find path to reference asset
 	UPackage* AssetOuterPackage = CastChecked<UPackage>(RelativeToAsset->GetOuter());
@@ -207,28 +204,62 @@ bool UFractureToolAutoUV::SaveGeneratedTexture(UTexture2D* GeneratedTexture, FSt
 	}
 
 	FString NewAssetPath = FPaths::Combine(PackageFolderPath, ObjectBaseName);
+
+	UTexture2D* GeneratedTexture = nullptr;
+	bool bNeedsNewPackage = true;
 	if (bAllowReplace)
 	{
 		// Modifying the static mesh in place. Delete existing asset so that we can have a clean duplicate
 		bool bNewAssetExists = UEditorAssetLibrary::DoesAssetExist(NewAssetPath);
 		if (bNewAssetExists)
 		{
-			bool bDeleteOK = UEditorAssetLibrary::DeleteAsset(NewAssetPath);
-			ensure(bDeleteOK);
+			UObject* OldObject = UEditorAssetLibrary::LoadAsset(NewAssetPath);
+			if (UTexture2D* OldTexture = Cast<UTexture2D>(OldObject))
+			{
+				FTexture2DBuilder TextureBuilder;
+				TextureBuilder.InitializeOver(OldTexture, FTexture2DBuilder::ETextureType::Color, ImageBuilder.GetDimensions());
+				TextureBuilder.Copy(ImageBuilder);
+				TextureBuilder.Commit(false);
+				GeneratedTexture = TextureBuilder.GetTexture2D();
+				FTexture2DBuilder::CopyPlatformDataToSourceData(GeneratedTexture, FTexture2DBuilder::ETextureType::Color);
+				bNeedsNewPackage = false;
+			}
+			else // old asset was wrong type; delete to replace
+			{
+				bool bDeleteOK = UEditorAssetLibrary::DeleteAsset(NewAssetPath);
+				ensure(bDeleteOK);
+			}
 		}
 	}
+	if (!GeneratedTexture)
+	{
+		FTexture2DBuilder TextureBuilder;
+		TextureBuilder.Initialize(FTexture2DBuilder::ETextureType::Color, ImageBuilder.GetDimensions());
+		TextureBuilder.Copy(ImageBuilder);
+		TextureBuilder.Commit(false);
+		GeneratedTexture = TextureBuilder.GetTexture2D();
+		FTexture2DBuilder::CopyPlatformDataToSourceData(GeneratedTexture, FTexture2DBuilder::ETextureType::Color);
+	}
+	AutoUVSettings->Result = GeneratedTexture;
+	check(GeneratedTexture);
+	check(GeneratedTexture->Source.IsValid());	// texture needs to have valid source data to be savd
 
-	// create new package
-	FString UniqueAssetName;
-	FString UniquePackageName;
+	if (bNeedsNewPackage)
+	{
+		check(GeneratedTexture->GetOuter() == GetTransientPackage());
+
+		// create new package
+		FString UniqueAssetName;
+		FString UniquePackageName;
 
-	FAssetToolsModule& AssetToolsModule = FModuleManager::LoadModuleChecked<FAssetToolsModule>("AssetTools");
-	AssetToolsModule.Get().CreateUniqueAssetName(NewAssetPath, TEXT(""), UniquePackageName, UniqueAssetName);
+		FAssetToolsModule& AssetToolsModule = FModuleManager::LoadModuleChecked<FAssetToolsModule>("AssetTools");
+		AssetToolsModule.Get().CreateUniqueAssetName(NewAssetPath, TEXT(""), UniquePackageName, UniqueAssetName);
 
-	UPackage* AssetPackage = CreatePackage(*UniquePackageName);
+		UPackage* AssetPackage = CreatePackage(*UniquePackageName);
 
-	// move texture from Transient package to real package
-	GeneratedTexture->Rename(*UniqueAssetName, AssetPackage, REN_None);
+		// move texture from Transient package to real package
+		GeneratedTexture->Rename(*UniqueAssetName, AssetPackage, REN_None);
+	}
 	// remove transient flag, add public/standalone/transactional
 	GeneratedTexture->ClearFlags(RF_Transient);
 	GeneratedTexture->SetFlags(RF_Public | RF_Standalone | RF_Transactional);
@@ -272,9 +303,8 @@ int32 UFractureToolAutoUV::ExecuteFracture(const FFractureToolContext& FractureC
 
 		UVTask.EnterProgressFrame(1, LOCTEXT("TexturingSurfaces", "Texturing internal surfaces"));
 
-		FTexture2DBuilder TextureBuilder;
+		
 		FImageDimensions Dimensions(OutputRes, OutputRes);
-		TextureBuilder.Initialize(FTexture2DBuilder::ETextureType::Color, Dimensions);
 
 		UE::Geometry::TImageBuilder<FVector4f> ImageBuilder;
 		ImageBuilder.SetDimensions(Dimensions);
@@ -318,20 +348,15 @@ int32 UFractureToolAutoUV::ExecuteFracture(const FFractureToolContext& FractureC
 
 		UVTask.EnterProgressFrame(1, LOCTEXT("SavingTexture", "Saving result"));
 
-		TextureBuilder.Copy(ImageBuilder);
-		TextureBuilder.Commit(false);
-		AutoUVSettings->Result = TextureBuilder.GetTexture2D();
-
 		// choose default texture name based on corresponding geometry collection name
 		FString BaseName = FractureContext.GetFracturedGeometryCollection()->GetName();
-		FTexture2DBuilder::CopyPlatformDataToSourceData(AutoUVSettings->Result, FTexture2DBuilder::ETextureType::Color);
 		FString Suffix = "_AutoUV";
 		if (AutoUVSettings->BakeTextureType == ETextureType::SpatialGradients)
 		{
 			Suffix = "_AutoUV_Spatial";
 		}
 		Suffix = FPaths::MakeValidFileName(Suffix);
-		SaveGeneratedTexture(AutoUVSettings->Result, FString::Printf(TEXT("%s%s"), *BaseName, *Suffix), FractureContext.GetFracturedGeometryCollection(), AutoUVSettings->bPromptToSave, AutoUVSettings->bReplaceExisting);
+		SaveGeneratedTexture(ImageBuilder, FString::Printf(TEXT("%s%s"), *BaseName, *Suffix), FractureContext.GetFracturedGeometryCollection(), AutoUVSettings->bPromptToSave, AutoUVSettings->bReplaceExisting);
 	}
 
 	return INDEX_NONE;
diff --git a/Engine/Plugins/Experimental/ChaosEditor/Source/FractureEditor/Private/FractureToolUV.h b/Engine/Plugins/Experimental/ChaosEditor/Source/FractureEditor/Private/FractureToolUV.h
--- a/Engine/Plugins/Experimental/ChaosEditor/Source/FractureEditor/Private/FractureToolUV.h
+++ b/Engine/Plugins/Experimental/ChaosEditor/Source/FractureEditor/Private/FractureToolUV.h
@@ -5,6 +5,8 @@
 #include "FractureTool.h"
 #include "GeometryCollection/GeometryCollection.h"
 
+#include "Image/ImageBuilder.h"
+
 #include "FractureToolUV.generated.h"
 
 class FFractureToolContext;
@@ -191,5 +193,5 @@ protected:
 	UPROPERTY(EditAnywhere, Category = Slicing)
 	TObjectPtr<UFractureAutoUVSettings> AutoUVSettings;
 
-	bool SaveGeneratedTexture(UTexture2D* GeneratedTexture, FString ObjectBaseName, const UObject* RelativeToAsset, bool bPromptToSave, bool bAllowReplace);
+	bool SaveGeneratedTexture(UE::Geometry::TImageBuilder<FVector4f>& ImageBuilder, FString ObjectBaseName, const UObject* RelativeToAsset, bool bPromptToSave, bool bAllowReplace);
 };
diff --git a/Engine/Plugins/Experimental/MeshModelingToolset/Source/ModelingComponents/Private/AssetUtils/Texture2DBuilder.cpp b/Engine/Plugins/Experimental/MeshModelingToolset/Source/ModelingComponents/Private/AssetUtils/Texture2DBuilder.cpp
--- a/Engine/Plugins/Experimental/MeshModelingToolset/Source/ModelingComponents/Private/AssetUtils/Texture2DBuilder.cpp
+++ b/Engine/Plugins/Experimental/MeshModelingToolset/Source/ModelingComponents/Private/AssetUtils/Texture2DBuilder.cpp
@@ -8,19 +8,71 @@ using namespace UE::Geometry;
 
 
 bool FTexture2DBuilder::Initialize(ETextureType BuildTypeIn, FImageDimensions DimensionsIn)
+{
+	// create new texture
+	EPixelFormat UsePixelFormat = (BuildTypeIn == ETextureType::EmissiveHDR) ? PF_FloatRGBA : PF_B8G8R8A8;
+	UTexture2D* TransientTexture = UTexture2D::CreateTransient((int32)DimensionsIn.GetWidth(), (int32)DimensionsIn.GetHeight(), UsePixelFormat);
+	return InitializeInternal(BuildTypeIn, DimensionsIn, TransientTexture);
+}
+
+
+bool FTexture2DBuilder::InitializeOver(UTexture2D* ExistingTexture, ETextureType BuildTypeIn, FImageDimensions DimensionsIn)
+{
+	// adapted from UTexture2D::CreateTransient
+
+	// create new texture
+	EPixelFormat InFormat = (BuildTypeIn == ETextureType::EmissiveHDR) ? PF_FloatRGBA : PF_B8G8R8A8;
+
+	int32 InSizeX = (int32)DimensionsIn.GetWidth(), InSizeY = (int32)DimensionsIn.GetHeight();
+
+	UTexture2D* NewTexture = nullptr;
+	if (ensureMsgf(InSizeX > 0 && InSizeY > 0 &&
+		(InSizeX % GPixelFormats[InFormat].BlockSizeX) == 0 &&
+		(InSizeY % GPixelFormats[InFormat].BlockSizeY) == 0, TEXT("Invalid size and/or pixel format for new texture")))
+	{
+		NewTexture = NewObject<UTexture2D>(
+			ExistingTexture->GetOuter(),
+			ExistingTexture->GetFName(),
+			ExistingTexture->GetFlags()
+			);
+
+		NewTexture->SetPlatformData(new FTexturePlatformData());
+		NewTexture->GetPlatformData()->SizeX = InSizeX;
+		NewTexture->GetPlatformData()->SizeY = InSizeY;
+		NewTexture->GetPlatformData()->PixelFormat = InFormat;
+
+		// Allocate first mipmap.
+		int32 NumBlocksX = InSizeX / GPixelFormats[InFormat].BlockSizeX;
+		int32 NumBlocksY = InSizeY / GPixelFormats[InFormat].BlockSizeY;
+		FTexture2DMipMap* Mip = new FTexture2DMipMap();
+		NewTexture->GetPlatformData()->Mips.Add(Mip);
+		Mip->SizeX = InSizeX;
+		Mip->SizeY = InSizeY;
+		Mip->BulkData.Lock(LOCK_READ_WRITE);
+		Mip->BulkData.Realloc(NumBlocksX * NumBlocksY * GPixelFormats[InFormat].BlockBytes);
+		Mip->BulkData.Unlock();
+	}
+	else
+	{
+		return false;
+	}
+
+	return InitializeInternal(BuildTypeIn, DimensionsIn, NewTexture);
+}
+
+
+bool FTexture2DBuilder::InitializeInternal(ETextureType BuildTypeIn, FImageDimensions DimensionsIn, UTexture2D* CreatedTextureIn)
 {
 	check(DimensionsIn.IsSquare());
 	BuildType = BuildTypeIn;
 	Dimensions = DimensionsIn;
 
-	// create new texture
-	EPixelFormat UsePixelFormat = (BuildTypeIn == ETextureType::EmissiveHDR) ? PF_FloatRGBA : PF_B8G8R8A8;
-	RawTexture2D = UTexture2D::CreateTransient((int32)Dimensions.GetWidth(), (int32)Dimensions.GetHeight(), UsePixelFormat);
+	RawTexture2D = CreatedTextureIn;
 	if (RawTexture2D == nullptr)
 	{
 		return false;
 	}
-	CurrentPixelFormat = UsePixelFormat;
+	CurrentPixelFormat = CreatedTextureIn->GetPixelFormat();
 
 	if (BuildType == ETextureType::ColorLinear 
 		|| BuildType == ETextureType::Roughness 
diff --git a/Engine/Plugins/Experimental/MeshModelingToolset/Source/ModelingComponents/Public/AssetUtils/Texture2DBuilder.h b/Engine/Plugins/Experimental/MeshModelingToolset/Source/ModelingComponents/Public/AssetUtils/Texture2DBuilder.h
--- a/Engine/Plugins/Experimental/MeshModelingToolset/Source/ModelingComponents/Public/AssetUtils/Texture2DBuilder.h
+++ b/Engine/Plugins/Experimental/MeshModelingToolset/Source/ModelingComponents/Public/AssetUtils/Texture2DBuilder.h
@@ -107,6 +107,12 @@ public:
 	bool Initialize(UTexture2D* ExistingTexture, ETextureType BuildTypeIn, bool bLockForEditing = true);
 
 
+	/**
+	 * Create a new UTexture2D 'over' an existing UTexture2D, with the given BuildType and Dimensions
+	 */
+	bool InitializeOver(UTexture2D* ExistingTexture, ETextureType BuildTypeIn, FImageDimensions DimensionsIn);
+
+
 	/**
 	 * Lock the Mip 0 buffer for editing
 	 * @return true if lock was successfull. Will return false if already locked.
@@ -330,6 +336,8 @@ public:
 		return TexBuilder.GetTexture2D();
 	}
 
+private:
+	bool InitializeInternal(ETextureType BuildTypeIn, FImageDimensions DimensionsIn, UTexture2D* CreatedTextureIn);
 };
 
 
