Better naming for Texture2DBuilder function that initializes a texture to replace an existing texture
(InitializeOver to InitializeAndReplaceExistingTexture)

#rnx

diff --git a/Engine/Plugins/Experimental/ChaosEditor/Source/FractureEditor/Private/FractureToolUV.cpp b/Engine/Plugins/Experimental/ChaosEditor/Source/FractureEditor/Private/FractureToolUV.cpp
--- a/Engine/Plugins/Experimental/ChaosEditor/Source/FractureEditor/Private/FractureToolUV.cpp
+++ b/Engine/Plugins/Experimental/ChaosEditor/Source/FractureEditor/Private/FractureToolUV.cpp
@@ -217,7 +217,7 @@ bool UFractureToolAutoUV::SaveGeneratedTexture(UE::Geometry::TImageBuilder<FVect
 			if (UTexture2D* OldTexture = Cast<UTexture2D>(OldObject))
 			{
 				FTexture2DBuilder TextureBuilder;
-				TextureBuilder.InitializeOver(OldTexture, FTexture2DBuilder::ETextureType::Color, ImageBuilder.GetDimensions());
+				TextureBuilder.InitializeAndReplaceExistingTexture(OldTexture, FTexture2DBuilder::ETextureType::Color, ImageBuilder.GetDimensions());
 				TextureBuilder.Copy(ImageBuilder);
 				TextureBuilder.Commit(false);
 				GeneratedTexture = TextureBuilder.GetTexture2D();
diff --git a/Engine/Plugins/Experimental/MeshModelingToolset/Source/ModelingComponents/Private/AssetUtils/Texture2DBuilder.cpp b/Engine/Plugins/Experimental/MeshModelingToolset/Source/ModelingComponents/Private/AssetUtils/Texture2DBuilder.cpp
--- a/Engine/Plugins/Experimental/MeshModelingToolset/Source/ModelingComponents/Private/AssetUtils/Texture2DBuilder.cpp
+++ b/Engine/Plugins/Experimental/MeshModelingToolset/Source/ModelingComponents/Private/AssetUtils/Texture2DBuilder.cpp
@@ -16,7 +16,7 @@ bool FTexture2DBuilder::Initialize(ETextureType BuildTypeIn, FImageDimensions Di
 }
 
 
-bool FTexture2DBuilder::InitializeOver(UTexture2D* ExistingTexture, ETextureType BuildTypeIn, FImageDimensions DimensionsIn)
+bool FTexture2DBuilder::InitializeAndReplaceExistingTexture(UTexture2D* ExistingTexture, ETextureType BuildTypeIn, FImageDimensions DimensionsIn)
 {
 	// adapted from UTexture2D::CreateTransient
 
diff --git a/Engine/Plugins/Experimental/MeshModelingToolset/Source/ModelingComponents/Public/AssetUtils/Texture2DBuilder.h b/Engine/Plugins/Experimental/MeshModelingToolset/Source/ModelingComponents/Public/AssetUtils/Texture2DBuilder.h
--- a/Engine/Plugins/Experimental/MeshModelingToolset/Source/ModelingComponents/Public/AssetUtils/Texture2DBuilder.h
+++ b/Engine/Plugins/Experimental/MeshModelingToolset/Source/ModelingComponents/Public/AssetUtils/Texture2DBuilder.h
@@ -110,7 +110,7 @@ public:
 	/**
 	 * Create a new UTexture2D 'over' an existing UTexture2D, with the given BuildType and Dimensions
 	 */
-	bool InitializeOver(UTexture2D* ExistingTexture, ETextureType BuildTypeIn, FImageDimensions DimensionsIn);
+	bool InitializeAndReplaceExistingTexture(UTexture2D* ExistingTexture, ETextureType BuildTypeIn, FImageDimensions DimensionsIn);
 
 
 	/**
