Add ability to add/remove UV channels from a geometry collection in the AutoUV tool

#rnx

diff --git a/Engine/Plugins/Experimental/ChaosEditor/Source/FractureEditor/Private/FractureToolUV.cpp b/Engine/Plugins/Experimental/ChaosEditor/Source/FractureEditor/Private/FractureToolUV.cpp
--- a/Engine/Plugins/Experimental/ChaosEditor/Source/FractureEditor/Private/FractureToolUV.cpp
+++ b/Engine/Plugins/Experimental/ChaosEditor/Source/FractureEditor/Private/FractureToolUV.cpp
@@ -49,6 +49,16 @@ int32 UFractureAutoUVSettings::GetSelectedChannelIndex(bool bForceToZeroOnFailur
 	return FoundIndex;
 }
 
+void UFractureAutoUVSettings::ChangeNumUVChannels(int32 Delta)
+{
+	int32 Target = UVChannelNamesList.Num() + Delta;
+	if (Target > 0 && Target < GeometryCollectionUV::MAX_NUM_UV_CHANNELS)
+	{
+		UFractureToolAutoUV* AutoUVTool = Cast<UFractureToolAutoUV>(OwnerTool.Get());
+		AutoUVTool->UpdateUVChannels(Target);
+	}
+}
+
 
 UFractureToolAutoUV::UFractureToolAutoUV(const FObjectInitializer& ObjInit)
 	: Super(ObjInit)
@@ -96,6 +106,11 @@ TArray<UObject*> UFractureToolAutoUV::GetSettingsObjects() const
 }
 
 void UFractureToolAutoUV::FractureContextChanged()
+{
+	UpdateUVChannels(-1);
+}
+
+void UFractureToolAutoUV::UpdateUVChannels(int32 TargetNumUVChannels)
 {
 	int32 MinUVChannels = GeometryCollectionUV::MAX_NUM_UV_CHANNELS;
 
@@ -107,8 +122,23 @@ void UFractureToolAutoUV::FractureContextChanged()
 		MinUVChannels = FMath::Min(MinUVChannels, NumChannels);
 	}
 
-	AutoUVSettings->SetNumUVChannels(MinUVChannels);
+	if (TargetNumUVChannels > -1 && TargetNumUVChannels != MinUVChannels)
+	{
+		bool bIsIncreasing = TargetNumUVChannels > MinUVChannels;
+		for (UGeometryCollectionComponent* GeometryCollectionComponent : GeomCompSelection)
+		{
+			int32 NumChannels = GeometryCollectionComponent->GetRestCollection()->GetGeometryCollection()->NumUVLayers();
+			// not if num channels is decreasing, all channels will be >= min channels so we apply the size change to all
+			// but if we're increasing the number of channels, some geometry collections may have more channels already and we leave those alone
+			if (NumChannels < TargetNumUVChannels || !bIsIncreasing)
+			{
+				GeometryCollectionComponent->GetRestCollection()->GetGeometryCollection()->SetNumUVLayers(TargetNumUVChannels);
+			}
+		}
+		MinUVChannels = TargetNumUVChannels;
+	}
 
+	AutoUVSettings->SetNumUVChannels(MinUVChannels);
 }
 
 void UFractureToolAutoUV::Render(const FSceneView* View, FViewport* Viewport, FPrimitiveDrawInterface* PDI)
diff --git a/Engine/Plugins/Experimental/ChaosEditor/Source/FractureEditor/Private/FractureToolUV.h b/Engine/Plugins/Experimental/ChaosEditor/Source/FractureEditor/Private/FractureToolUV.h
--- a/Engine/Plugins/Experimental/ChaosEditor/Source/FractureEditor/Private/FractureToolUV.h
+++ b/Engine/Plugins/Experimental/ChaosEditor/Source/FractureEditor/Private/FractureToolUV.h
@@ -71,6 +71,20 @@ public:
 	void SetNumUVChannels(int32 NumUVChannels);
 	int32 GetSelectedChannelIndex(bool bForceToZeroOnFailure = true);
 
+	UFUNCTION(CallInEditor, Category = UVChannel, meta = (DisplayName = "Add UV Channel"))
+	void AddUVChannel()
+	{
+		ChangeNumUVChannels(1);
+	}
+
+	UFUNCTION(CallInEditor, Category = UVChannel, meta = (DisplayName = "Delete UV Channel"))
+	void DeleteUVChannel()
+	{
+		ChangeNumUVChannels(-1);
+	}
+
+	void ChangeNumUVChannels(int32 Delta);
+
 	/** Whether to layout UVs (to ensure unique texels for each face), or just use the existing UVs */
 	UPROPERTY(EditAnywhere, Category = Atlas, meta = (DisplayName = "Do UV Layout"))
 	bool bDoUVLayout = true;
@@ -189,6 +203,9 @@ public:
 
 	virtual TArray<FFractureToolContext> GetFractureToolContexts() const override;
 
+	// Update number of UV channels for selected geometry collections
+	void UpdateUVChannels(int32 TargetNumUVChannels);
+
 
 protected:
 	UPROPERTY(EditAnywhere, Category = Slicing)
diff --git a/Engine/Source/Runtime/Experimental/Chaos/Private/GeometryCollection/GeometryCollection.cpp b/Engine/Source/Runtime/Experimental/Chaos/Private/GeometryCollection/GeometryCollection.cpp
--- a/Engine/Source/Runtime/Experimental/Chaos/Private/GeometryCollection/GeometryCollection.cpp
+++ b/Engine/Source/Runtime/Experimental/Chaos/Private/GeometryCollection/GeometryCollection.cpp
@@ -1229,6 +1229,19 @@ int32 FGeometryCollection::NumUVLayers() const
 	return MinUVLayers;
 }
 
+bool FGeometryCollection::SetNumUVLayers(int32 NumLayers)
+{
+	if (NumLayers < 1 || NumLayers >= GeometryCollectionUV::MAX_NUM_UV_CHANNELS)
+	{
+		return false;
+	}
+
+	for (int32 VertIdx = 0; VertIdx < NumElements(FGeometryCollection::VerticesGroup); ++VertIdx)
+	{
+		UVs[VertIdx].SetNum(NumLayers);
+	}
+	return true;
+}
 
 FGeometryCollection* FGeometryCollection::NewGeometryCollection(const TArray<float>& RawVertexArray, const TArray<int32>& RawIndicesArray, bool ReverseVertexOrder)
 {
diff --git a/Engine/Source/Runtime/Experimental/Chaos/Public/GeometryCollection/GeometryCollection.h b/Engine/Source/Runtime/Experimental/Chaos/Public/GeometryCollection/GeometryCollection.h
--- a/Engine/Source/Runtime/Experimental/Chaos/Public/GeometryCollection/GeometryCollection.h
+++ b/Engine/Source/Runtime/Experimental/Chaos/Public/GeometryCollection/GeometryCollection.h
@@ -219,6 +219,9 @@ public:
 	/** Returns number of UV layers represented by UV array. A Valid Geometry Collection has the same count for every vertex */
 	int32 NumUVLayers() const;
 
+	/** Update a geometry collection to have the target number of UV layers (must be in the range [1, MAX_UV_LAYERS)) */
+	bool SetNumUVLayers(int32 NumLayers);
+
 	FORCEINLINE bool IsGeometry(int32 Element) const { return TransformToGeometryIndex[Element] != INDEX_NONE; }
 	FORCEINLINE bool IsClustered(int32 Element) const { const TManagedArray<int32>& SimType = SimulationType;  return !!(SimType[Element] == ESimulationTypes::FST_Clustered); }
 	FORCEINLINE bool IsRigid(int32 Element) const { const TManagedArray<int32>& SimType = SimulationType;  return !!(SimType[Element] == ESimulationTypes::FST_Rigid); }
