make auto uv new channels fill with channel 0

diff --git a/Engine/Plugins/Experimental/ChaosEditor/Source/FractureEditor/Private/FractureToolUV.cpp b/Engine/Plugins/Experimental/ChaosEditor/Source/FractureEditor/Private/FractureToolUV.cpp
--- a/Engine/Plugins/Experimental/ChaosEditor/Source/FractureEditor/Private/FractureToolUV.cpp
+++ b/Engine/Plugins/Experimental/ChaosEditor/Source/FractureEditor/Private/FractureToolUV.cpp
@@ -133,6 +133,18 @@ void UFractureToolAutoUV::UpdateUVChannels(int32 TargetNumUVChannels)
 			if (NumChannels < TargetNumUVChannels || !bIsIncreasing)
 			{
 				GeometryCollectionComponent->GetRestCollection()->GetGeometryCollection()->SetNumUVLayers(TargetNumUVChannels);
+				if (bIsIncreasing)
+				{
+					// Copy UV layer 0 into new UV layer
+					FGeometryCollection& Collection = *GeometryCollectionComponent->GetRestCollection()->GetGeometryCollection();
+					for (int32 Idx = 0; Idx < Collection.UVs.Num(); Idx++)
+					{
+						for (int32 Ch = NumChannels; Ch < TargetNumUVChannels; Ch++)
+						{
+							Collection.UVs[Idx][Ch] = Collection.UVs[Idx][0];
+						}
+					}
+				}
 			}
 		}
 		MinUVChannels = TargetNumUVChannels;
