fracture autouv ui improvements: change from "skip layout" checkbox to two buttons for layout/texture + add projection scale options

#rnx

diff --git a/Engine/Plugins/Experimental/ChaosEditor/Source/FractureEditor/Private/FractureToolUV.cpp b/Engine/Plugins/Experimental/ChaosEditor/Source/FractureEditor/Private/FractureToolUV.cpp
--- a/Engine/Plugins/Experimental/ChaosEditor/Source/FractureEditor/Private/FractureToolUV.cpp
+++ b/Engine/Plugins/Experimental/ChaosEditor/Source/FractureEditor/Private/FractureToolUV.cpp
@@ -65,6 +65,18 @@ void UFractureAutoUVSettings::BoxProjectUVs()
 	AutoUVTool->BoxProjectUVs();
 }
 
+void UFractureAutoUVSettings::LayoutUVs()
+{
+	UFractureToolAutoUV* AutoUVTool = Cast<UFractureToolAutoUV>(OwnerTool.Get());
+	AutoUVTool->LayoutUVs();
+}
+
+void UFractureAutoUVSettings::BakeTexture()
+{
+	UFractureToolAutoUV* AutoUVTool = Cast<UFractureToolAutoUV>(OwnerTool.Get());
+	AutoUVTool->BakeTexture();
+}
+
 UFractureToolAutoUV::UFractureToolAutoUV(const FObjectInitializer& ObjInit)
 	: Super(ObjInit)
 {
@@ -173,7 +185,7 @@ void UFractureToolAutoUV::BoxProjectUVs()
 	{
 		UE::PlanarCut::BoxProjectUVs(UVLayer,
 			*GeometryCollectionComponent->GetRestCollection()->GetGeometryCollection(),
-			FVector3d(100, 100, 100),
+			(FVector3d)AutoUVSettings->ProjectionScale,
 			UseMaterialIDs,
 			AutoUVSettings->TargetMaterialIDs == ETargetMaterialIDs::OddIDs ? EmptyMaterialIDs : AutoUVSettings->MaterialIDs);
 
@@ -345,93 +357,148 @@ bool UFractureToolAutoUV::SaveGeneratedTexture(UE::Geometry::TImageBuilder<FVect
 	return true;
 }
 
+namespace
+{
+	UE::PlanarCut::EUseMaterials GetUseMaterials(ETargetMaterialIDs TargetIDs)
+	{
+		return TargetIDs == ETargetMaterialIDs::SelectedIDs ? UE::PlanarCut::EUseMaterials::NoDefaultMaterials :
+			(TargetIDs == ETargetMaterialIDs::AllIDs ? UE::PlanarCut::EUseMaterials::AllMaterials : UE::PlanarCut::EUseMaterials::OddMaterials);
+	}
+}
 
-int32 UFractureToolAutoUV::ExecuteFracture(const FFractureToolContext& FractureContext)
+void UFractureToolAutoUV::LayoutUVs()
 {
-	if (FractureContext.GetGeometryCollection().IsValid())
+	TSet<UGeometryCollectionComponent*> GeomCompSelection;
+	GetSelectedGeometryCollectionComponents(GeomCompSelection);
+	for (UGeometryCollectionComponent* GeometryCollectionComponent : GeomCompSelection)
 	{
-		FScopedSlowTask UVTask(AutoUVSettings->bDoUVLayout ? 3 : 2, LOCTEXT("StartingAutoUV", "Automatically laying out and texturing internal surfaces"));
-		UVTask.MakeDialog();
+		LayoutUVsForComponent(GeometryCollectionComponent);
 
-		FGeometryCollection& Collection = *FractureContext.GetGeometryCollection();
-		
-		int32 OutputRes = (int32)AutoUVSettings->Resolution;
-		TArray<int32> EmptyMaterialIDs;
+		GeometryCollectionComponent->MarkRenderDynamicDataDirty();
+		GeometryCollectionComponent->MarkRenderStateDirty();
+	}
+}
 
-		int32 UVLayer = AutoUVSettings->GetSelectedChannelIndex();
 
-		UE::PlanarCut::EUseMaterials UseMaterialIDs =
-			AutoUVSettings->TargetMaterialIDs == ETargetMaterialIDs::SelectedIDs ? UE::PlanarCut::EUseMaterials::NoDefaultMaterials :
-			(AutoUVSettings->TargetMaterialIDs == ETargetMaterialIDs::AllIDs ? UE::PlanarCut::EUseMaterials::AllMaterials : UE::PlanarCut::EUseMaterials::OddMaterials);
+bool UFractureToolAutoUV::LayoutUVsForComponent(UGeometryCollectionComponent* Component)
+{
+	int32 OutputRes = (int32)AutoUVSettings->Resolution;
+	TArray<int32> EmptyMaterialIDs;
 
-		if (AutoUVSettings->bDoUVLayout)
-		{
-			UVTask.EnterProgressFrame(1, LOCTEXT("LayOutUVIslands", "Laying out UV islands"));
-			if (!UE::PlanarCut::UVLayout(UVLayer, Collection, OutputRes, AutoUVSettings->GutterSize,
-				UseMaterialIDs,
-				AutoUVSettings->TargetMaterialIDs == ETargetMaterialIDs::OddIDs ? EmptyMaterialIDs : AutoUVSettings->MaterialIDs))
-			{
-				// failed to do layout
-				return INDEX_NONE;
-			}
-		}
+	int32 UVLayer = AutoUVSettings->GetSelectedChannelIndex();
 
-		UVTask.EnterProgressFrame(1, LOCTEXT("TexturingSurfaces", "Texturing internal surfaces"));
+	UE::PlanarCut::EUseMaterials UseMaterialIDs = GetUseMaterials(AutoUVSettings->TargetMaterialIDs);
 
-		
-		FImageDimensions Dimensions(OutputRes, OutputRes);
+	bool bLayoutSuccess =
+		UE::PlanarCut::UVLayout(UVLayer, *Component->GetRestCollection()->GetGeometryCollection(), OutputRes, AutoUVSettings->GutterSize,
+			UseMaterialIDs,
+			AutoUVSettings->TargetMaterialIDs == ETargetMaterialIDs::OddIDs ? EmptyMaterialIDs : AutoUVSettings->MaterialIDs);
+	
+	return bLayoutSuccess;
+}
 
-		UE::Geometry::TImageBuilder<FVector4f> ImageBuilder;
-		ImageBuilder.SetDimensions(Dimensions);
-		ImageBuilder.Clear(FVector4f(0, 0, 0, 0));
 
-		typedef UE::PlanarCut::EBakeAttributes EBakeAttributes;
-		FIndex4i Attributes;
-		// Note: Ordering of these attributes should match the order and comments in the AutoUVSettings struct
-		//		 Update the order and comments there if you change the ordering here.
-		if (AutoUVSettings->BakeTextureType == ETextureType::ThicknessAndSurfaceAttributes)
-		{
-			Attributes = FIndex4i(
-				AutoUVSettings->bDistToOuter ? (int32)EBakeAttributes::DistanceToExternal : 0,
-				AutoUVSettings->bAmbientOcclusion ? (int32)EBakeAttributes::AmbientOcclusion : 0,
-				AutoUVSettings->bSmoothedCurvature ? (int32)EBakeAttributes::Curvature : 0,
-				0
-			);
-		}
-		else
-		{
-			Attributes = FIndex4i(
-				(int32)EBakeAttributes::PositionX,
-				(int32)EBakeAttributes::PositionY,
-				(int32)EBakeAttributes::PositionZ,
-				0
-			);
-		}
-		UE::PlanarCut::FTextureAttributeSettings AttribSettings;
-		AttribSettings.ToExternal_MaxDistance = AutoUVSettings->MaxDistance;
-		AttribSettings.AO_Rays = AutoUVSettings->OcclusionRays;
-		AttribSettings.AO_BlurRadius = AutoUVSettings->OcclusionBlurRadius;
-		AttribSettings.Curvature_BlurRadius = AutoUVSettings->CurvatureBlurRadius;
-		AttribSettings.Curvature_SmoothingSteps = AutoUVSettings->SmoothingIterations;
-		AttribSettings.Curvature_VoxelRes = AutoUVSettings->VoxelResolution;
-		AttribSettings.Curvature_ThicknessFactor = AutoUVSettings->ThicknessFactor;
-		AttribSettings.Curvature_MaxValue = AutoUVSettings->MaxCurvature;
-		AttribSettings.ClearGutterChannel = 3; // default clear the gutters for the alpha channel, so it shows more clearly the island boundaries
-		UE::PlanarCut::TextureInternalSurfaces(UVLayer, Collection, FMath::CeilToInt(AutoUVSettings->GutterSize), Attributes, AttribSettings, ImageBuilder,
-			UseMaterialIDs,
-			AutoUVSettings->TargetMaterialIDs == ETargetMaterialIDs::OddIDs ? EmptyMaterialIDs: AutoUVSettings->MaterialIDs);
+void UFractureToolAutoUV::BakeTexture()
+{
+	TSet<UGeometryCollectionComponent*> GeomCompSelection;
+	GetSelectedGeometryCollectionComponents(GeomCompSelection);
+	for (UGeometryCollectionComponent* GeometryCollectionComponent : GeomCompSelection)
+	{
+		BakeTextureForComponent(GeometryCollectionComponent);
 
-		UVTask.EnterProgressFrame(1, LOCTEXT("SavingTexture", "Saving result"));
+		GeometryCollectionComponent->MarkRenderDynamicDataDirty();
+		GeometryCollectionComponent->MarkRenderStateDirty();
+	}
+}
 
-		// choose default texture name based on corresponding geometry collection name
-		FString BaseName = FractureContext.GetFracturedGeometryCollection()->GetName();
-		FString Suffix = "_AutoUV";
-		if (AutoUVSettings->BakeTextureType == ETextureType::SpatialGradients)
+void UFractureToolAutoUV::BakeTextureForComponent(UGeometryCollectionComponent* Component, TFunction<void(int32, const FText&)> Progress)
+{
+	FGeometryCollection& Collection = *Component->GetRestCollection()->GetGeometryCollection();
+
+	int32 OutputRes = (int32)AutoUVSettings->Resolution;
+	TArray<int32> EmptyMaterialIDs;
+
+	int32 UVLayer = AutoUVSettings->GetSelectedChannelIndex();
+
+	UE::PlanarCut::EUseMaterials UseMaterialIDs = GetUseMaterials(AutoUVSettings->TargetMaterialIDs);
+
+	FImageDimensions Dimensions(OutputRes, OutputRes);
+
+	UE::Geometry::TImageBuilder<FVector4f> ImageBuilder;
+	ImageBuilder.SetDimensions(Dimensions);
+	ImageBuilder.Clear(FVector4f(0, 0, 0, 0));
+
+	typedef UE::PlanarCut::EBakeAttributes EBakeAttributes;
+	FIndex4i Attributes;
+	// Note: Ordering of these attributes should match the order and comments in the AutoUVSettings struct
+	//		 Update the order and comments there if you change the ordering here.
+	if (AutoUVSettings->BakeTextureType == ETextureType::ThicknessAndSurfaceAttributes)
+	{
+		Attributes = FIndex4i(
+			AutoUVSettings->bDistToOuter ? (int32)EBakeAttributes::DistanceToExternal : 0,
+			AutoUVSettings->bAmbientOcclusion ? (int32)EBakeAttributes::AmbientOcclusion : 0,
+			AutoUVSettings->bSmoothedCurvature ? (int32)EBakeAttributes::Curvature : 0,
+			0
+		);
+	}
+	else
+	{
+		Attributes = FIndex4i(
+			(int32)EBakeAttributes::PositionX,
+			(int32)EBakeAttributes::PositionY,
+			(int32)EBakeAttributes::PositionZ,
+			0
+		);
+	}
+	UE::PlanarCut::FTextureAttributeSettings AttribSettings;
+	AttribSettings.ToExternal_MaxDistance = AutoUVSettings->MaxDistance;
+	AttribSettings.AO_Rays = AutoUVSettings->OcclusionRays;
+	AttribSettings.AO_BlurRadius = AutoUVSettings->OcclusionBlurRadius;
+	AttribSettings.Curvature_BlurRadius = AutoUVSettings->CurvatureBlurRadius;
+	AttribSettings.Curvature_SmoothingSteps = AutoUVSettings->SmoothingIterations;
+	AttribSettings.Curvature_VoxelRes = AutoUVSettings->VoxelResolution;
+	AttribSettings.Curvature_ThicknessFactor = AutoUVSettings->ThicknessFactor;
+	AttribSettings.Curvature_MaxValue = AutoUVSettings->MaxCurvature;
+	AttribSettings.ClearGutterChannel = 3; // default clear the gutters for the alpha channel, so it shows more clearly the island boundaries
+	UE::PlanarCut::TextureInternalSurfaces(UVLayer, Collection, FMath::CeilToInt(AutoUVSettings->GutterSize), Attributes, AttribSettings, ImageBuilder,
+		UseMaterialIDs,
+		AutoUVSettings->TargetMaterialIDs == ETargetMaterialIDs::OddIDs ? EmptyMaterialIDs : AutoUVSettings->MaterialIDs);
+
+	if (Progress)
+	{
+		Progress(1, LOCTEXT("SavingTexture", "Saving result"));
+	}
+
+	// choose default texture name based on corresponding geometry collection name
+	FString BaseName = Component->GetRestCollection()->GetName();
+	FString Suffix = "_AutoUV";
+	if (AutoUVSettings->BakeTextureType == ETextureType::SpatialGradients)
+	{
+		Suffix = "_AutoUV_Spatial";
+	}
+	Suffix = FPaths::MakeValidFileName(Suffix);
+	SaveGeneratedTexture(ImageBuilder, FString::Printf(TEXT("%s%s"), *BaseName, *Suffix), Component->GetRestCollection(), AutoUVSettings->bPromptToSave, AutoUVSettings->bReplaceExisting);
+}
+
+int32 UFractureToolAutoUV::ExecuteFracture(const FFractureToolContext& FractureContext)
+{
+	if (FractureContext.GetGeometryCollection().IsValid())
+	{
+		bool bDoUVLayout = true;
+
+		FScopedSlowTask UVTask(bDoUVLayout ? 3 : 2, LOCTEXT("StartingAutoUV", "Automatically laying out and texturing internal surfaces"));
+		UVTask.MakeDialog();
+
+		UGeometryCollectionComponent* Component = FractureContext.GetGeometryCollectionComponent();
+
+		if (bDoUVLayout)
 		{
-			Suffix = "_AutoUV_Spatial";
+			UVTask.EnterProgressFrame(1, LOCTEXT("LayOutUVIslands", "Laying out UV islands"));
+			LayoutUVsForComponent(Component);
 		}
-		Suffix = FPaths::MakeValidFileName(Suffix);
-		SaveGeneratedTexture(ImageBuilder, FString::Printf(TEXT("%s%s"), *BaseName, *Suffix), FractureContext.GetFracturedGeometryCollection(), AutoUVSettings->bPromptToSave, AutoUVSettings->bReplaceExisting);
+
+		UVTask.EnterProgressFrame(1, LOCTEXT("TexturingSurfaces", "Texturing internal surfaces"));
+		BakeTextureForComponent(Component, [&](int32 AmtWork, const FText& Msg) { UVTask.EnterProgressFrame(AmtWork, Msg); });
 	}
 
 	return INDEX_NONE;
diff --git a/Engine/Plugins/Experimental/ChaosEditor/Source/FractureEditor/Private/FractureToolUV.h b/Engine/Plugins/Experimental/ChaosEditor/Source/FractureEditor/Private/FractureToolUV.h
--- a/Engine/Plugins/Experimental/ChaosEditor/Source/FractureEditor/Private/FractureToolUV.h
+++ b/Engine/Plugins/Experimental/ChaosEditor/Source/FractureEditor/Private/FractureToolUV.h
@@ -85,12 +85,17 @@ public:
 
 	void ChangeNumUVChannels(int32 Delta);
 
-	UFUNCTION(CallInEditor, Category = UVIslands, meta = (DisplayName = "Box Project UVs"))
-	void BoxProjectUVs();
+	UFUNCTION(CallInEditor, Category = Operations, meta = (DisplayName = "Layout UVs", DisplayPriority = 1))
+	void LayoutUVs();
+
+	UFUNCTION(CallInEditor, Category = Operations, meta = (DisplayName = "Bake Texture", DisplayPriority = 2))
+	void BakeTexture();
 
-	/** Whether to layout UVs (to ensure unique texels for each face), or just use the existing UVs */
-	UPROPERTY(EditAnywhere, Category = Atlas, meta = (DisplayName = "Do UV Layout"))
-	bool bDoUVLayout = true;
+	UPROPERTY(EditAnywhere, Category = Unwrap)
+	FVector ProjectionScale = FVector(100, 100, 100);
+
+	UFUNCTION(CallInEditor, Category = Unwrap, meta = (DisplayName = "Box Project UVs"))
+	void BoxProjectUVs();
 
 	/** The pixel resolution of the generated map */
 	UPROPERTY(EditAnywhere, Category = MapSettings)
@@ -211,6 +216,9 @@ public:
 
 	// Create new UV islands via box projection for the selected geometry collections
 	void BoxProjectUVs();
+	
+	void LayoutUVs();
+	void BakeTexture();
 
 
 protected:
@@ -218,4 +226,7 @@ protected:
 	TObjectPtr<UFractureAutoUVSettings> AutoUVSettings;
 
 	bool SaveGeneratedTexture(UE::Geometry::TImageBuilder<FVector4f>& ImageBuilder, FString ObjectBaseName, const UObject* RelativeToAsset, bool bPromptToSave, bool bAllowReplace);
+
+	bool LayoutUVsForComponent(UGeometryCollectionComponent* Component);
+	void BakeTextureForComponent(UGeometryCollectionComponent* Component, TFunction<void(int32, const FText&)> Progress = nullptr);
 };
