#ue5 - build fix: LWC: fix ambiguous CeilToInt() call on an int32 argument (just don't use the function now).

diff --git a/Engine/Plugins/Experimental/ChaosEditor/Source/FractureEditor/Private/FractureToolUV.cpp b/Engine/Plugins/Experimental/ChaosEditor/Source/FractureEditor/Private/FractureToolUV.cpp
--- a/Engine/Plugins/Experimental/ChaosEditor/Source/FractureEditor/Private/FractureToolUV.cpp
+++ b/Engine/Plugins/Experimental/ChaosEditor/Source/FractureEditor/Private/FractureToolUV.cpp
@@ -460,7 +460,7 @@ void UFractureToolAutoUV::BakeTextureForComponent(UGeometryCollectionComponent*
 	AttribSettings.Curvature_ThicknessFactor = AutoUVSettings->ThicknessFactor;
 	AttribSettings.Curvature_MaxValue = AutoUVSettings->MaxCurvature;
 	AttribSettings.ClearGutterChannel = 3; // default clear the gutters for the alpha channel, so it shows more clearly the island boundaries
-	UE::PlanarCut::TextureInternalSurfaces(UVLayer, Collection, FMath::CeilToInt(AutoUVSettings->GutterSize), Attributes, AttribSettings, ImageBuilder,
+	UE::PlanarCut::TextureInternalSurfaces(UVLayer, Collection, AutoUVSettings->GutterSize, Attributes, AttribSettings, ImageBuilder,
 		UseMaterialIDs,
 		AutoUVSettings->TargetMaterialIDs == ETargetMaterialIDs::OddIDs ? EmptyMaterialIDs : AutoUVSettings->MaterialIDs);
 
