FractureMode
-- Fixed a transform mistake in EmbedFlush tool.

diff --git a/Engine/Plugins/Experimental/ChaosEditor/Source/FractureEditor/Private/FractureToolEmbed.cpp b/Engine/Plugins/Experimental/ChaosEditor/Source/FractureEditor/Private/FractureToolEmbed.cpp
--- a/Engine/Plugins/Experimental/ChaosEditor/Source/FractureEditor/Private/FractureToolEmbed.cpp
+++ b/Engine/Plugins/Experimental/ChaosEditor/Source/FractureEditor/Private/FractureToolEmbed.cpp
@@ -21,8 +21,6 @@
 
 #define LOCTEXT_NAMESPACE "FractureToolEmbed"
 
-
-
 FText UFractureToolAddEmbeddedGeometry::GetDisplayText() const
 {
 	return FText(NSLOCTEXT("Fracture", "FractureToolAddEmbeddedGeometry", "Embed"));
@@ -422,7 +420,7 @@ void UFractureToolFlushEmbeddedGeometry::Execute(TWeakPtr<FFractureEditorModeToo
 
 								for (const FTransform& Transform : It->Value)
 								{
-									ActorFactory->CreateActor(LoadedExemplarAsset, TargetWorld->GetCurrentLevel(), ComponentSpace*Transform, SpawnParameters);
+									ActorFactory->CreateActor(LoadedExemplarAsset, TargetWorld->GetCurrentLevel(), Transform * ComponentSpace, SpawnParameters);
 								}
 							}
 
