FractureTool
-- Fixed crash when generating new GeometryCollection if Size array isn't present

diff --git a/Engine/Plugins/Experimental/ChaosEditor/Source/FractureEditor/Private/FractureEditorModeToolkit.cpp b/Engine/Plugins/Experimental/ChaosEditor/Source/FractureEditor/Private/FractureEditorModeToolkit.cpp
--- a/Engine/Plugins/Experimental/ChaosEditor/Source/FractureEditor/Private/FractureEditorModeToolkit.cpp
+++ b/Engine/Plugins/Experimental/ChaosEditor/Source/FractureEditor/Private/FractureEditorModeToolkit.cpp
@@ -1133,6 +1133,11 @@ void FFractureEditorModeToolkit::UpdateGeometryComponentAttributes(UGeometryColl
 			}
 
 			// Normalize Size
+			if (!GeometryCollection->HasAttribute("Size", FTransformCollection::TransformGroup))
+			{
+				GeometryCollection->AddAttribute<float>("Size", FTransformCollection::TransformGroup);
+				UE_LOG(LogFractureTool, Warning, TEXT("Added Size attribute to GeometryCollection."));
+			}
 			TManagedArray<float>& Size = GeometryCollection->GetAttribute<float>("Size", FTransformCollection::TransformGroup);
 			const TManagedArray<int32>& SimulationType = GeometryCollection->SimulationType;
 			float MaxSize = 0.0f;
