Fracture Tool
-- Changed Geometry Colection generation to prioritize selected components over actors, which makes converting BPs with multiple components viable source for GeometryCollections.

#okforhithub public, internal

diff --git a/Engine/Plugins/Experimental/ChaosEditor/Source/FractureEditor/Private/FractureToolGenerators.cpp b/Engine/Plugins/Experimental/ChaosEditor/Source/FractureEditor/Private/FractureToolGenerators.cpp
--- a/Engine/Plugins/Experimental/ChaosEditor/Source/FractureEditor/Private/FractureToolGenerators.cpp
+++ b/Engine/Plugins/Experimental/ChaosEditor/Source/FractureEditor/Private/FractureToolGenerators.cpp
@@ -171,80 +171,74 @@ AGeometryCollectionActor* UFractureToolGenerateAsset::ConvertActorsToGeometryCol
 
 	FGeometryCollectionEdit GeometryCollectionEdit = NewActor->GetGeometryCollectionComponent()->EditRestCollection(GeometryCollection::EEditUpdate::RestPhysicsDynamic);
 	UGeometryCollection* FracturedGeometryCollection = GeometryCollectionEdit.GetRestCollection();
+	check(FracturedGeometryCollection);
 
 	for (AActor* Actor : Actors)
 	{
 		const FTransform ActorTransform(Actor->GetTransform());
 		const FVector ActorOffset(Actor->GetActorLocation() - FirstActor->GetActorLocation());
 
-		check(FracturedGeometryCollection);
-
-		if (const AStaticMeshActor* StaticMeshActor = Cast<AStaticMeshActor>(Actor))
+		TArray<UStaticMeshComponent*> StaticMeshComponents;
+		Actor->GetComponents<UStaticMeshComponent>(StaticMeshComponents, true);
+		for (int32 ii = 0, ni = StaticMeshComponents.Num(); ii < ni; ++ii)
 		{
-			TArray<UStaticMeshComponent*> StaticMeshComponents;
-			Actor->GetComponents<UStaticMeshComponent>(StaticMeshComponents, true);
-			for (int32 ii = 0, ni = StaticMeshComponents.Num(); ii < ni; ++ii)
+			// We're partial to static mesh components, here
+			UStaticMeshComponent* StaticMeshComponent = StaticMeshComponents[ii];
+			if (StaticMeshComponent != nullptr)
 			{
-				// We're partial to static mesh components, here
-				UStaticMeshComponent* StaticMeshComponent = StaticMeshComponents[ii];
-				if (StaticMeshComponent != nullptr)
+				UStaticMesh* ComponentStaticMesh = StaticMeshComponent->GetStaticMesh();
+				if (ComponentStaticMesh != nullptr)
 				{
-					UStaticMesh* ComponentStaticMesh = StaticMeshComponent->GetStaticMesh();
-					if (ComponentStaticMesh != nullptr)
-					{
-						// If any of the static meshes have Nanite enabled, also enable on the new geometry collection asset for convenience.
-						FracturedGeometryCollection->EnableNanite |= ComponentStaticMesh->NaniteSettings.bEnabled;
-					}
+					// If any of the static meshes have Nanite enabled, also enable on the new geometry collection asset for convenience.
+					FracturedGeometryCollection->EnableNanite |= ComponentStaticMesh->NaniteSettings.bEnabled;
+				}
 
-					FTransform ComponentTransform(StaticMeshComponent->GetComponentTransform());
-					ComponentTransform.SetTranslation((ComponentTransform.GetTranslation() - ActorTransform.GetTranslation()) + ActorOffset);
+				FTransform ComponentTransform(StaticMeshComponent->GetComponentTransform());
+				ComponentTransform.SetTranslation((ComponentTransform.GetTranslation() - ActorTransform.GetTranslation()) + ActorOffset);
 
-					// Record the contributing source on the asset.
-					FSoftObjectPath SourceSoftObjectPath(ComponentStaticMesh);
-					decltype(FGeometryCollectionSource::SourceMaterial) SourceMaterials(StaticMeshComponent->GetMaterials());
-					FracturedGeometryCollection->GeometrySource.Add({ SourceSoftObjectPath, ComponentTransform, SourceMaterials });
+				// Record the contributing source on the asset.
+				FSoftObjectPath SourceSoftObjectPath(ComponentStaticMesh);
+				decltype(FGeometryCollectionSource::SourceMaterial) SourceMaterials(StaticMeshComponent->GetMaterials());
+				FracturedGeometryCollection->GeometrySource.Add({ SourceSoftObjectPath, ComponentTransform, SourceMaterials });
 
-					FGeometryCollectionConversion::AppendStaticMesh(ComponentStaticMesh, SourceMaterials, ComponentTransform, FracturedGeometryCollection, true);
-				}
+				FGeometryCollectionConversion::AppendStaticMesh(ComponentStaticMesh, SourceMaterials, ComponentTransform, FracturedGeometryCollection, true);
 			}
 		}
-		else if (const AGeometryCollectionActor* GeometryActor = Cast<AGeometryCollectionActor>(Actor))
+
+		TArray<UGeometryCollectionComponent*> GeometryCollectionComponents;
+		Actor->GetComponents<UGeometryCollectionComponent>(GeometryCollectionComponents, true);
+		for (int32 ii = 0, ni = GeometryCollectionComponents.Num(); ii < ni; ++ii)
 		{
-			TArray<UGeometryCollectionComponent*> GeometryCollectionComponents;
-			Actor->GetComponents<UGeometryCollectionComponent>(GeometryCollectionComponents, true);
-			for (int32 ii = 0, ni = GeometryCollectionComponents.Num(); ii < ni; ++ii)
+			UGeometryCollectionComponent* GeometryCollectionComponent = GeometryCollectionComponents[ii];
+			if (GeometryCollectionComponent != nullptr)
 			{
-				UGeometryCollectionComponent* GeometryCollectionComponent = GeometryCollectionComponents[ii];
-				if (GeometryCollectionComponent != nullptr)
+				const UGeometryCollection* RestCollection = GeometryCollectionComponent->GetRestCollection();
+				if (RestCollection != nullptr)
 				{
-					const UGeometryCollection* RestCollection = GeometryCollectionComponent->GetRestCollection();
-					if (RestCollection != nullptr)
-					{
-						// If any of the static meshes have Nanite enabled, also enable on the new geometry collection asset for convenience.
-						FracturedGeometryCollection->EnableNanite |= RestCollection->EnableNanite;
-					}
+					// If any of the static meshes have Nanite enabled, also enable on the new geometry collection asset for convenience.
+					FracturedGeometryCollection->EnableNanite |= RestCollection->EnableNanite;
+				}
 
-					FTransform ComponentTransform(GeometryCollectionComponent->GetComponentTransform());
-					ComponentTransform.SetTranslation((ComponentTransform.GetTranslation() - ActorTransform.GetTranslation()) + ActorOffset);
+				FTransform ComponentTransform(GeometryCollectionComponent->GetComponentTransform());
+				ComponentTransform.SetTranslation((ComponentTransform.GetTranslation() - ActorTransform.GetTranslation()) + ActorOffset);
 
-					// Record the contributing source on the asset.
-					FSoftObjectPath SourceSoftObjectPath(RestCollection);
+				// Record the contributing source on the asset.
+				FSoftObjectPath SourceSoftObjectPath(RestCollection);
 
-					// We're not interested in recording the final material of the collection since it's inevitably the Selection material.
-					int32 NumMaterials = GeometryCollectionComponent->GetNumMaterials() - 1;
-					TArray<TObjectPtr<UMaterialInterface>> SourceMaterials;
-					SourceMaterials.SetNum(NumMaterials);
-					for (int32 MaterialIndex = 0; MaterialIndex < NumMaterials; ++MaterialIndex)
-					{
-						SourceMaterials[MaterialIndex] = GeometryCollectionComponent->GetMaterial(MaterialIndex);
-					}
-					FracturedGeometryCollection->GeometrySource.Add({ SourceSoftObjectPath, ComponentTransform, SourceMaterials });
+				// We're not interested in recording the final material of the collection since it's inevitably the Selection material.
+				int32 NumMaterials = GeometryCollectionComponent->GetNumMaterials() - 1;
+				TArray<TObjectPtr<UMaterialInterface>> SourceMaterials;
+				SourceMaterials.SetNum(NumMaterials);
+				for (int32 MaterialIndex = 0; MaterialIndex < NumMaterials; ++MaterialIndex)
+				{
+					SourceMaterials[MaterialIndex] = GeometryCollectionComponent->GetMaterial(MaterialIndex);
+				}
+				FracturedGeometryCollection->GeometrySource.Add({ SourceSoftObjectPath, ComponentTransform, SourceMaterials });
 
-					FGeometryCollectionConversion::AppendGeometryCollection(RestCollection, GeometryCollectionComponent, ComponentTransform, FracturedGeometryCollection, false);
+				FGeometryCollectionConversion::AppendGeometryCollection(RestCollection, GeometryCollectionComponent, ComponentTransform, FracturedGeometryCollection, false);
 
-				}
 			}
-		}	
+		}
 	}
 
 	FracturedGeometryCollection->InitializeMaterials();
