Updating fracture mode to utilize the multiple tab support of FAssetEditorModeUILayer

diff --git a/Engine/Plugins/Experimental/ChaosEditor/Source/FractureEditor/Private/FractureEditorModeToolkit.cpp b/Engine/Plugins/Experimental/ChaosEditor/Source/FractureEditor/Private/FractureEditorModeToolkit.cpp
--- a/Engine/Plugins/Experimental/ChaosEditor/Source/FractureEditor/Private/FractureEditorModeToolkit.cpp
+++ b/Engine/Plugins/Experimental/ChaosEditor/Source/FractureEditor/Private/FractureEditorModeToolkit.cpp
@@ -63,6 +63,7 @@
 #include "PropertyHandle.h"
 
 #include "FractureSettings.h"
+#include "Toolkits/AssetEditorModeUILayer.h"
 
 #define LOCTEXT_NAMESPACE "FFractureEditorModeToolkit"
 
@@ -226,104 +227,6 @@ void FFractureEditorModeToolkit::Init(const TSharedPtr<IToolkitHost>& InitToolki
 
 	float Padding = 4.0f;
 	FMargin MorePadding = FMargin(10.0f, 2.0f);
-
-	TSharedRef<SExpandableArea> HistogramExpander = SNew(SExpandableArea)
-	.AreaTitle(FText(LOCTEXT("Histogram", "Histogram")))
-	.HeaderPadding(FMargin(2.0, 2.0))
-	.Padding(MorePadding)
-	.BorderBackgroundColor(FLinearColor(.6, .6, .6, 1.0f))
-	.BodyBorderBackgroundColor(FLinearColor(1.0, 0.0, 0.0))
-	.AreaTitleFont(FEditorStyle::Get().GetFontStyle("HistogramDetailsView.CategoryFontStyle"))
-	.InitiallyCollapsed(true)
-	.BodyContent()
-	[
-		SNew(SVerticalBox)
-		+ SVerticalBox::Slot()
-		.AutoHeight()
-		[
-			HistogramDetailsView.ToSharedRef()
-		]
-		+ SVerticalBox::Slot()
-		.AutoHeight()
-		[
-			SNew(SButton)
-			//.ButtonStyle(FAppStyle::Get(), "PrimaryButton")
-			.TextStyle(FAppStyle::Get(), "ButtonText")
-			.HAlign(HAlign_Center)
-			.ContentPadding(FMargin(10.f, Padding))
-			.OnClicked(this, &FFractureEditorModeToolkit::ResetHistogramSelection)
-			.IsEnabled(this, &FFractureEditorModeToolkit::CanResetFilter)
-			.Text(LOCTEXT("ResetFilterButton", "Reset Filter"))
-		]
-		+ SVerticalBox::Slot()
-		[
-			SAssignNew(HistogramView, SGeometryCollectionHistogram)
-			.OnBoneSelectionChanged(this, &FFractureEditorModeToolkit::OnHistogramBoneSelectionChanged)
-		]
-	];
-
-	TSharedRef<SExpandableArea> OutlinerExpander = SNew(SExpandableArea)
-	.AreaTitle(FText(LOCTEXT("Outliner", "Outliner")))
-	.HeaderPadding(FMargin(2.0, 2.0))
-	.Padding(MorePadding)
-	.BorderImage(FEditorStyle::Get().GetBrush("DetailsView.CategoryTop"))
-	.BorderBackgroundColor( FLinearColor( .6,.6,.6, 1.0f ) )
-	.BodyBorderBackgroundColor (FLinearColor( 1.0, 0.0, 0.0))
-	.AreaTitleFont(FEditorStyle::Get().GetFontStyle("DetailsView.CategoryFontStyle"))
-	.BodyContent()
-	[
-		SNew(SVerticalBox)
-		+SVerticalBox::Slot()
-		[
-			SNew(SSplitter)
-			.Orientation(Orient_Vertical)
-			+ SSplitter::Slot()
-			.SizeRule(TAttribute<SSplitter::ESizeRule>::Create([this, HistogramExpander]() {
-				return HistogramExpander->IsExpanded() ? SSplitter::ESizeRule::FractionOfParent : SSplitter::ESizeRule::SizeToContent;
-				}))
-			.Value(1.f)
-			[
-				HistogramExpander
-			]
-			+ SSplitter::Slot()
-			.SizeRule(SSplitter::ESizeRule::FractionOfParent)
-			.Value(1.f)
-			[
-				SNew(SVerticalBox)
-				+ SVerticalBox::Slot()
-				.AutoHeight()
-				[
-					OutlinerDetailsView.ToSharedRef()
-				]
-				+ SVerticalBox::Slot()
-				[
-					SAssignNew(OutlinerView, SGeometryCollectionOutliner)
-					.OnBoneSelectionChanged(this, &FFractureEditorModeToolkit::OnOutlinerBoneSelectionChanged)
-				]
-				+ SVerticalBox::Slot()
-				.AutoHeight()
-				[
-					SNew(STextBlock)
-					.Text(this, &FFractureEditorModeToolkit::GetSelectionInfo)
-				]
-			]
-		]	
-	];
-
-	TSharedRef<SExpandableArea> StatisticsExpander = SNew(SExpandableArea)
-	.AreaTitle(FText(LOCTEXT("LevelStatistics", "Level Statistics")))
-	.HeaderPadding(FMargin(2.0, 2.0))
-	.Padding(MorePadding)
-	.BorderImage(FEditorStyle::Get().GetBrush("DetailsView.CategoryTop"))
-	.BorderBackgroundColor( FLinearColor( .6,.6,.6, 1.0f ) )
-	.BodyBorderBackgroundColor (FLinearColor( 1.0, 0.0, 0.0))
-	.AreaTitleFont(FEditorStyle::Get().GetFontStyle("DetailsView.CategoryFontStyle"))
-	.BodyContent()
-	[
-		SNew(STextBlock)
-		.Text(this, &FFractureEditorModeToolkit::GetStatisticsSummary)
-	];
-
 	SAssignNew(ToolkitWidget, SBox)
 	[
 		SNew(SVerticalBox)
@@ -388,24 +291,6 @@ void FFractureEditorModeToolkit::Init(const TSharedPtr<IToolkitHost>& InitToolki
 					]
 				]
 			]
-
-			+SSplitter::Slot()
-			.SizeRule( TAttribute<SSplitter::ESizeRule>::Create( [this, OutlinerExpander] () { 
-				return OutlinerExpander->IsExpanded() ? SSplitter::ESizeRule::FractionOfParent : SSplitter::ESizeRule::SizeToContent; 
-			} ) )
-			.Value(1.f)
-			[
-				OutlinerExpander
-			]
-
-			+SSplitter::Slot()
-			.SizeRule( TAttribute<SSplitter::ESizeRule>::Create( [this, StatisticsExpander] () { 
-				return StatisticsExpander->IsExpanded() ? SSplitter::ESizeRule::FractionOfParent : SSplitter::ESizeRule::SizeToContent; 
-			} ) )
-			.Value(0.25f)
-			[
-				StatisticsExpander
-			]
 		]
 	];
 
@@ -420,6 +305,176 @@ void FFractureEditorModeToolkit::Init(const TSharedPtr<IToolkitHost>& InitToolki
 
 }
 
+void FFractureEditorModeToolkit::RequestModeUITabs()
+{
+	FModeToolkit::RequestModeUITabs();
+	if (ModeUILayer.IsValid())
+	{
+		TSharedPtr<FAssetEditorModeUILayer> ModeUILayerPtr = ModeUILayer.Pin();
+		TSharedRef<FWorkspaceItem> MenuGroup = ModeUILayerPtr->GetModeMenuCategory().ToSharedRef();
+		HierarchyTabInfo.OnSpawnTab = FOnSpawnTab::CreateSP(SharedThis(this), &FFractureEditorModeToolkit::CreateHierarchyTab);
+		HierarchyTabInfo.TabLabel = LOCTEXT("FractureHierarchy", "Fracture Hierarchy");
+		HierarchyTabInfo.TabTooltip = LOCTEXT("ModesToolboxTabTooltipText", "Open the  Modes tab, which contains the active editor mode's settings.");
+		ModeUILayerPtr->SetModePanelInfo(FAssetEditorModeUILayer::TopRightTabID, HierarchyTabInfo);
+
+
+		StatisticsTabInfo.OnSpawnTab = FOnSpawnTab::CreateSP(SharedThis(this), &FFractureEditorModeToolkit::CreateStatisticsTab);
+		StatisticsTabInfo.TabLabel = LOCTEXT("FractureStatistics", "Level Statistics");
+		StatisticsTabInfo.TabTooltip = LOCTEXT("ModesToolboxTabTooltipText", "Open the  Modes tab, which contains the active editor mode's settings.");
+		ModeUILayerPtr->SetModePanelInfo(FAssetEditorModeUILayer::BottomLeftTabID, StatisticsTabInfo);
+	}
+}
+
+void FFractureEditorModeToolkit::InvokeUI()
+{
+	FModeToolkit::InvokeUI();
+
+	if (ModeUILayer.IsValid())
+	{
+		TSharedPtr<FAssetEditorModeUILayer> ModeUILayerPtr = ModeUILayer.Pin();
+		HierarchyTab = ModeUILayerPtr->GetTabManager()->TryInvokeTab(FAssetEditorModeUILayer::TopRightTabID);
+		StatisticsTab = ModeUILayerPtr->GetTabManager()->TryInvokeTab(FAssetEditorModeUILayer::BottomLeftTabID);
+	}
+}
+
+TSharedRef<SDockTab> FFractureEditorModeToolkit::CreateHierarchyTab(const FSpawnTabArgs& Args)
+{
+	float Padding = 4.0f;
+	FMargin MorePadding = FMargin(10.0f, 2.0f);
+
+	TSharedRef<SExpandableArea> HistogramExpander = SNew(SExpandableArea)
+		.AreaTitle(FText(LOCTEXT("Histogram", "Histogram")))
+		.HeaderPadding(FMargin(2.0, 2.0))
+		.Padding(MorePadding)
+		.BorderBackgroundColor(FLinearColor(.6, .6, .6, 1.0f))
+		.BodyBorderBackgroundColor(FLinearColor(1.0, 0.0, 0.0))
+		.AreaTitleFont(FEditorStyle::Get().GetFontStyle("HistogramDetailsView.CategoryFontStyle"))
+		.InitiallyCollapsed(true)
+		.BodyContent()
+		[
+			SNew(SVerticalBox)
+			+ SVerticalBox::Slot()
+			.AutoHeight()
+			[
+				HistogramDetailsView.ToSharedRef()
+			]
+			+ SVerticalBox::Slot()
+			.AutoHeight()
+			[
+				SNew(SButton)
+				//.ButtonStyle(FAppStyle::Get(), "PrimaryButton")
+				.TextStyle(FAppStyle::Get(), "ButtonText")
+				.HAlign(HAlign_Center)
+				.ContentPadding(FMargin(10.f, Padding))
+				.OnClicked(this, &FFractureEditorModeToolkit::ResetHistogramSelection)
+				.IsEnabled(this, &FFractureEditorModeToolkit::CanResetFilter)
+				.Text(LOCTEXT("ResetFilterButton", "Reset Filter"))
+			]
+			+ SVerticalBox::Slot()
+			[
+				SAssignNew(HistogramView, SGeometryCollectionHistogram)
+				.OnBoneSelectionChanged(this, &FFractureEditorModeToolkit::OnHistogramBoneSelectionChanged)
+			]
+		];
+
+	TSharedRef<SExpandableArea> OutlinerExpander = SNew(SExpandableArea)
+		.AreaTitle(FText(LOCTEXT("Outliner", "Outliner")))
+		.HeaderPadding(FMargin(2.0, 2.0))
+		.Padding(MorePadding)
+		.BorderImage(FEditorStyle::Get().GetBrush("DetailsView.CategoryTop"))
+		.BorderBackgroundColor(FLinearColor(.6, .6, .6, 1.0f))
+		.BodyBorderBackgroundColor(FLinearColor(1.0, 0.0, 0.0))
+		.AreaTitleFont(FEditorStyle::Get().GetFontStyle("DetailsView.CategoryFontStyle"))
+		.BodyContent()
+		[
+			SNew(SVerticalBox)
+
+			+ SVerticalBox::Slot()
+			[
+				SNew(SSplitter)
+				.Orientation(Orient_Vertical)
+			
+				+ SSplitter::Slot()
+				.SizeRule(TAttribute<SSplitter::ESizeRule>::Create([this, HistogramExpander]() {
+					return HistogramExpander->IsExpanded() ? SSplitter::ESizeRule::FractionOfParent : SSplitter::ESizeRule::SizeToContent;
+				}))
+				.Value(1.f)
+				[
+					HistogramExpander
+				]
+
+				+ SSplitter::Slot()
+				.SizeRule(SSplitter::ESizeRule::FractionOfParent)
+				.Value(1.f)
+				[
+					SNew(SVerticalBox)
+
+					+ SVerticalBox::Slot()
+					.AutoHeight()
+					[
+						OutlinerDetailsView.ToSharedRef()
+					]
+
+					+ SVerticalBox::Slot()
+					[
+						SAssignNew(OutlinerView, SGeometryCollectionOutliner)
+						.OnBoneSelectionChanged(this, &FFractureEditorModeToolkit::OnOutlinerBoneSelectionChanged)
+					]
+
+					+ SVerticalBox::Slot()
+					.AutoHeight()
+					[
+						SNew(STextBlock)
+						.Text(this, &FFractureEditorModeToolkit::GetSelectionInfo)
+					]
+				]
+			]
+		];
+
+	TSharedPtr<SDockTab> CreatedTab = SNew(SDockTab)
+		[
+			OutlinerExpander
+		];
+
+	const FSlateBrush* TabIcon = GetEditorModeIcon().GetSmallIcon();
+	if (CreatedTab)
+	{
+		CreatedTab->SetTabIcon(TabIcon);
+	}
+	HierarchyTab = CreatedTab;
+	return CreatedTab.ToSharedRef();
+}
+
+TSharedRef<SDockTab> FFractureEditorModeToolkit::CreateStatisticsTab(const FSpawnTabArgs& Args)
+{
+	FMargin MorePadding = FMargin(10.0f, 2.0f);
+	TSharedRef<SExpandableArea> StatisticsExpander = SNew(SExpandableArea)
+		.AreaTitle(FText(LOCTEXT("LevelStatistics", "Level Statistics")))
+		.HeaderPadding(FMargin(2.0, 2.0))
+		.Padding(MorePadding)
+		.BorderImage(FEditorStyle::Get().GetBrush("DetailsView.CategoryTop"))
+		.BorderBackgroundColor(FLinearColor(.6, .6, .6, 1.0f))
+		.BodyBorderBackgroundColor(FLinearColor(1.0, 0.0, 0.0))
+		.AreaTitleFont(FEditorStyle::Get().GetFontStyle("DetailsView.CategoryFontStyle"))
+		.BodyContent()
+		[
+			SNew(STextBlock)
+			.Text(this, &FFractureEditorModeToolkit::GetStatisticsSummary)
+		];
+	TSharedPtr<SDockTab> CreatedTab = SNew(SDockTab)
+		[
+			StatisticsExpander
+		];
+
+	const FSlateBrush* TabIcon = GetEditorModeIcon().GetSmallIcon();
+	if (CreatedTab)
+	{
+		CreatedTab->SetTabIcon(TabIcon);
+	}
+	StatisticsTab = CreatedTab;
+	return CreatedTab.ToSharedRef();
+}
+
 void FFractureEditorModeToolkit::OnObjectPostEditChange( UObject* Object, FPropertyChangedEvent& PropertyChangedEvent )
 {
 	if (PropertyChangedEvent.Property)
diff --git a/Engine/Plugins/Experimental/ChaosEditor/Source/FractureEditor/Public/FractureEditorModeToolkit.h b/Engine/Plugins/Experimental/ChaosEditor/Source/FractureEditor/Public/FractureEditorModeToolkit.h
--- a/Engine/Plugins/Experimental/ChaosEditor/Source/FractureEditor/Public/FractureEditorModeToolkit.h
+++ b/Engine/Plugins/Experimental/ChaosEditor/Source/FractureEditor/Public/FractureEditorModeToolkit.h
@@ -83,7 +83,6 @@ public:
 	/** FModeToolkit interface */
 	virtual void Init(const TSharedPtr<IToolkitHost>& InitToolkitHost) override;
 	virtual void Init(const TSharedPtr<IToolkitHost>& InitToolkitHost, TWeakObjectPtr<UEdMode> InOwningMode) override;
-
 	/** IToolkit interface */
 	virtual FName GetToolkitFName() const override;
 	virtual FText GetBaseToolkitName() const override;
@@ -174,6 +173,12 @@ public:
 	TSharedPtr<SWidget> ShowBoneColorsWidget;
 
 protected:
+	/** FModeToolkit interface */
+	virtual void RequestModeUITabs() override;
+	virtual void InvokeUI() override;
+
+	TSharedRef<SDockTab> CreateHierarchyTab(const FSpawnTabArgs& Args);
+	TSharedRef<SDockTab> CreateStatisticsTab(const FSpawnTabArgs& Args);
 	static bool IsGeometryCollectionSelected();
 	static bool IsSelectedActorsInEditorWorld();	
 
@@ -195,5 +200,9 @@ private:
 	TSharedPtr<SWidget> ToolkitWidget;
 	TSharedPtr<SGeometryCollectionOutliner> OutlinerView;
 	TSharedPtr<SGeometryCollectionHistogram> HistogramView;
+	TWeakPtr<SDockTab> HierarchyTab;
+	FMinorTabConfig HierarchyTabInfo;
+	TWeakPtr<SDockTab> StatisticsTab;
+	FMinorTabConfig StatisticsTabInfo;
 
 };
