FractureTool
-- Added Hide/Unhide tools

diff --git a/Engine/Plugins/Experimental/ChaosCaching/Source/ChaosCaching/Private/Chaos/Adapters/StaticMeshComponentCacheAdapter.cpp b/Engine/Plugins/Experimental/ChaosCaching/Source/ChaosCaching/Private/Chaos/Adapters/StaticMeshComponentCacheAdapter.cpp
--- a/Engine/Plugins/Experimental/ChaosCaching/Source/ChaosCaching/Private/Chaos/Adapters/StaticMeshComponentCacheAdapter.cpp
+++ b/Engine/Plugins/Experimental/ChaosCaching/Source/ChaosCaching/Private/Chaos/Adapters/StaticMeshComponentCacheAdapter.cpp
@@ -177,6 +177,8 @@ namespace Chaos
 
 	bool FStaticMeshCacheAdapter::InitializeForPlayback(UPrimitiveComponent* InComponent, UChaosCache* InCache, float InTime)
 	{
+		ensure(IsInGameThreadContext());
+
 #if WITH_CHAOS
 
 		if(Cast<UStaticMeshComponent>(InComponent) && InComponent->GetBodyInstance() && InComponent->GetBodyInstance()->ActorHandle)
diff --git a/Engine/Plugins/Experimental/ChaosEditor/Source/FractureEditor/Private/FractureEditorCommands.h b/Engine/Plugins/Experimental/ChaosEditor/Source/FractureEditor/Private/FractureEditorCommands.h
--- a/Engine/Plugins/Experimental/ChaosEditor/Source/FractureEditor/Private/FractureEditorCommands.h
+++ b/Engine/Plugins/Experimental/ChaosEditor/Source/FractureEditor/Private/FractureEditorCommands.h
@@ -37,6 +37,8 @@ class FFractureEditorCommands : public TCommands<FFractureEditorCommands>
 
 		// Edit Commands
 		TSharedPtr< FUICommandInfo > DeleteBranch;
+		TSharedPtr< FUICommandInfo > Hide;
+		TSharedPtr< FUICommandInfo > Unhide;
 		
 		// Generate Commands
 		TSharedPtr< FUICommandInfo > GenerateAsset;
diff --git a/Engine/Plugins/Experimental/ChaosEditor/Source/FractureEditor/Private/FractureEditorModeToolkit.cpp b/Engine/Plugins/Experimental/ChaosEditor/Source/FractureEditor/Private/FractureEditorModeToolkit.cpp
--- a/Engine/Plugins/Experimental/ChaosEditor/Source/FractureEditor/Private/FractureEditorModeToolkit.cpp
+++ b/Engine/Plugins/Experimental/ChaosEditor/Source/FractureEditor/Private/FractureEditorModeToolkit.cpp
@@ -563,6 +563,8 @@ void FFractureEditorModeToolkit::BuildToolPalette(FName PaletteIndex, class FToo
 	else if (PaletteIndex == TEXT("Edit"))
 	{
 		ToolbarBuilder.AddToolBarButton(Commands.DeleteBranch);
+		ToolbarBuilder.AddToolBarButton(Commands.Hide);
+		ToolbarBuilder.AddToolBarButton(Commands.Unhide);
 	}
 	else if (PaletteIndex == TEXT("Cluster"))
 	{
diff --git a/Engine/Plugins/Experimental/ChaosEditor/Source/FractureEditor/Private/FractureEditorStyle.cpp b/Engine/Plugins/Experimental/ChaosEditor/Source/FractureEditor/Private/FractureEditorStyle.cpp
--- a/Engine/Plugins/Experimental/ChaosEditor/Source/FractureEditor/Private/FractureEditorStyle.cpp
+++ b/Engine/Plugins/Experimental/ChaosEditor/Source/FractureEditor/Private/FractureEditorStyle.cpp
@@ -54,6 +54,8 @@ FFractureEditorStyle::FFractureEditorStyle()
 	Set("FractureEditor.AutoUV",               new IMAGE_BRUSH_SVG("FractureAutoUV", IconSize));
 
 	Set("FractureEditor.DeleteBranch",         new IMAGE_BRUSH_SVG("FracturePrune", IconSize));
+	Set("FractureEditor.Hide",				   new IMAGE_BRUSH_SVG("FractureHide", IconSize));
+	Set("FractureEditor.Unhide",               new IMAGE_BRUSH_SVG("FractureUnhide", IconSize));
 
 	Set("FractureEditor.RecomputeNormals",     new IMAGE_BRUSH_SVG("FractureNormals", IconSize));
 	Set("FractureEditor.Resample",             new IMAGE_BRUSH_SVG("FractureResample", IconSize));
diff --git a/Engine/Plugins/Experimental/ChaosEditor/Source/FractureEditor/Private/FractureToolEditing.cpp b/Engine/Plugins/Experimental/ChaosEditor/Source/FractureEditor/Private/FractureToolEditing.cpp
--- a/Engine/Plugins/Experimental/ChaosEditor/Source/FractureEditor/Private/FractureToolEditing.cpp
+++ b/Engine/Plugins/Experimental/ChaosEditor/Source/FractureEditor/Private/FractureToolEditing.cpp
@@ -99,6 +99,139 @@ void UFractureToolDeleteBranch::Execute(TWeakPtr<FFractureEditorModeToolkit> InT
 
 
 
+FText UFractureToolHide::GetDisplayText() const
+{
+	return FText(NSLOCTEXT("FractureToolEditingOps", "FractureToolHide", "Hide"));
+}
+
+FText UFractureToolHide::GetTooltipText() const
+{
+	return FText(NSLOCTEXT("FractureToolEditingOps", "FractureToolHideTooltip", "Set all geometry in selected branch to invisible."));
+}
+
+FSlateIcon UFractureToolHide::GetToolIcon() const
+{
+	return FSlateIcon("FractureEditorStyle", "FractureEditor.Hide");
+}
+
+void UFractureToolHide::RegisterUICommand(FFractureEditorCommands* BindingContext)
+{
+	UI_COMMAND_EXT(BindingContext, UICommandInfo, "Hide", "Hide", "Set all geometry in selected branch to invisible.", EUserInterfaceActionType::Button, FInputChord());
+	BindingContext->Hide = UICommandInfo;
+}
+
+void UFractureToolHide::Execute(TWeakPtr<FFractureEditorModeToolkit> InToolkit)
+{
+	if (InToolkit.IsValid())
+	{
+		FFractureEditorModeToolkit* Toolkit = InToolkit.Pin().Get();
+
+		TArray<FFractureToolContext> Contexts = GetFractureToolContexts();
+
+		for (FFractureToolContext& Context : Contexts)
+		{
+			FGeometryCollection* GeometryCollection = Context.GetGeometryCollection().Get();
+			UGeometryCollection* FracturedGeometryCollection = Context.GetFracturedGeometryCollection();
+
+			Context.ConvertSelectionToRigidNodes();
+
+			const TManagedArray<int32>&	TransformToGeometryIndex = GeometryCollection->TransformToGeometryIndex;
+			const TManagedArray<int32>&	FaceStart = GeometryCollection->FaceStart;
+			const TManagedArray<int32>&	FaceCount = GeometryCollection->FaceCount;
+			TManagedArray<bool>& Visible = GeometryCollection->Visible;
+
+
+			const TArray<int32>& Selection = Context.GetSelection();
+			for (int32 Idx : Selection)
+			{
+				// Iterate the faces in the geometry of this rigid node and set invisible.
+				if (TransformToGeometryIndex[Idx] > INDEX_NONE)
+				{
+					int32 CurrFace = FaceStart[TransformToGeometryIndex[Idx]];
+					for (int32 FaceOffset = 0; FaceOffset < FaceCount[TransformToGeometryIndex[Idx]]; ++FaceOffset)
+					{
+						Visible[CurrFace + FaceOffset] = false;
+					}
+				}
+			}
+
+			Context.GetGeometryCollectionComponent()->MarkRenderStateDirty();
+			Context.GetGeometryCollectionComponent()->MarkRenderDynamicDataDirty();
+			Refresh(Context, Toolkit, true);
+		}
+
+		SetOutlinerComponents(Contexts, Toolkit);
+	}
+}
+
+
+
+FText UFractureToolUnhide::GetDisplayText() const
+{
+	return FText(NSLOCTEXT("FractureToolEditingOps", "FractureToolUnhide", "Hide"));
+}
+
+FText UFractureToolUnhide::GetTooltipText() const
+{
+	return FText(NSLOCTEXT("FractureToolEditingOps", "FractureToolUnhideTooltip", "Set all geometry in selected branch to visible."));
+}
+
+FSlateIcon UFractureToolUnhide::GetToolIcon() const
+{
+	return FSlateIcon("FractureEditorStyle", "FractureEditor.Unhide");
+}
+
+void UFractureToolUnhide::RegisterUICommand(FFractureEditorCommands* BindingContext)
+{
+	UI_COMMAND_EXT(BindingContext, UICommandInfo, "Unhide", "Unhide", "Set all geometry in selected branch to visible.", EUserInterfaceActionType::Button, FInputChord());
+	BindingContext->Unhide = UICommandInfo;
+}
+
+void UFractureToolUnhide::Execute(TWeakPtr<FFractureEditorModeToolkit> InToolkit)
+{
+	if (InToolkit.IsValid())
+	{
+		FFractureEditorModeToolkit* Toolkit = InToolkit.Pin().Get();
+
+		TArray<FFractureToolContext> Contexts = GetFractureToolContexts();
+
+		for (FFractureToolContext& Context : Contexts)
+		{
+			FGeometryCollection* GeometryCollection = Context.GetGeometryCollection().Get();
+			UGeometryCollection* FracturedGeometryCollection = Context.GetFracturedGeometryCollection();
+
+			Context.ConvertSelectionToRigidNodes();
+
+			const TManagedArray<int32>& TransformToGeometryIndex = GeometryCollection->TransformToGeometryIndex;
+			const TManagedArray<int32>& FaceStart = GeometryCollection->FaceStart;
+			const TManagedArray<int32>& FaceCount = GeometryCollection->FaceCount;
+			TManagedArray<bool>& Visible = GeometryCollection->Visible;
+
+
+			const TArray<int32>& Selection = Context.GetSelection();
+			for (int32 Idx : Selection)
+			{
+				// Iterate the faces in the geometry of this rigid node and set invisible.
+				if (TransformToGeometryIndex[Idx] > INDEX_NONE)
+				{
+					int32 CurrFace = FaceStart[TransformToGeometryIndex[Idx]];
+					for (int32 FaceOffset = 0; FaceOffset < FaceCount[TransformToGeometryIndex[Idx]]; ++FaceOffset)
+					{
+						Visible[CurrFace + FaceOffset] = true;
+					}
+				}
+			}
+
+			Context.GetGeometryCollectionComponent()->MarkRenderStateDirty();
+			Context.GetGeometryCollectionComponent()->MarkRenderDynamicDataDirty();
+			Refresh(Context, Toolkit, true);
+		}
+
+		SetOutlinerComponents(Contexts, Toolkit);
+	}
+}
+
+
 FText UFractureToolValidate::GetDisplayText() const
 {
 	return FText(NSLOCTEXT("FractureToolEditingOps", "FractureToolValidate", "Validate"));
diff --git a/Engine/Plugins/Experimental/ChaosEditor/Source/FractureEditor/Private/FractureToolEditing.h b/Engine/Plugins/Experimental/ChaosEditor/Source/FractureEditor/Private/FractureToolEditing.h
--- a/Engine/Plugins/Experimental/ChaosEditor/Source/FractureEditor/Private/FractureToolEditing.h
+++ b/Engine/Plugins/Experimental/ChaosEditor/Source/FractureEditor/Private/FractureToolEditing.h
@@ -22,6 +22,42 @@ public:
 	virtual void Execute(TWeakPtr<FFractureEditorModeToolkit> InToolkit) override;
 };
 
+
+UCLASS(DisplayName = "Hide", Category = "FractureTools")
+class UFractureToolHide : public UFractureActionTool
+{
+public:
+	GENERATED_BODY()
+
+	UFractureToolHide(const FObjectInitializer& ObjInit) : Super(ObjInit) {}
+
+	// UFractureActionTool Interface
+	virtual FText GetDisplayText() const override;
+	virtual FText GetTooltipText() const override;
+	virtual FSlateIcon GetToolIcon() const override;
+	virtual void RegisterUICommand(FFractureEditorCommands* BindingContext) override;
+	virtual void Execute(TWeakPtr<FFractureEditorModeToolkit> InToolkit) override;
+};
+
+
+UCLASS(DisplayName = "Unhide", Category = "FractureTools")
+class UFractureToolUnhide : public UFractureActionTool
+{
+public:
+	GENERATED_BODY()
+
+	UFractureToolUnhide(const FObjectInitializer& ObjInit) : Super(ObjInit) {}
+
+	// UFractureActionTool Interface
+	virtual FText GetDisplayText() const override;
+	virtual FText GetTooltipText() const override;
+	virtual FSlateIcon GetToolIcon() const override;
+	virtual void RegisterUICommand(FFractureEditorCommands* BindingContext) override;
+	virtual void Execute(TWeakPtr<FFractureEditorModeToolkit> InToolkit) override;
+};
+
+
+
 UCLASS(DisplayName = "Validate", Category = "FractureTools")
 class UFractureToolValidate : public UFractureActionTool
 {
diff --git a/Engine/Plugins/Experimental/ChaosEditor/Source/FractureEditor/Private/SGeometryCollectionOutliner.cpp b/Engine/Plugins/Experimental/ChaosEditor/Source/FractureEditor/Private/SGeometryCollectionOutliner.cpp
--- a/Engine/Plugins/Experimental/ChaosEditor/Source/FractureEditor/Private/SGeometryCollectionOutliner.cpp
+++ b/Engine/Plugins/Experimental/ChaosEditor/Source/FractureEditor/Private/SGeometryCollectionOutliner.cpp
@@ -591,7 +591,14 @@ TSharedRef<ITableRow> FGeometryCollectionTreeItemBone::MakeTreeRowWidget(const T
 				break;
 
 			case FGeometryCollection::ESimulationTypes::FST_Rigid:
-				TextColor = FSlateColor::UseForeground();
+				if (GeometryCollectionPtr->IsVisible(GetBoneIndex()))
+				{ 
+					TextColor = FSlateColor::UseForeground();
+				}
+				else
+				{
+					TextColor = FLinearColor{0.1f, 0.1f, 0.1f, 1.f};
+				}
 				break;
 
 			case FGeometryCollection::ESimulationTypes::FST_Clustered:
diff --git a/Engine/Source/Runtime/Experimental/Chaos/Private/GeometryCollection/GeometryCollection.cpp b/Engine/Source/Runtime/Experimental/Chaos/Private/GeometryCollection/GeometryCollection.cpp
--- a/Engine/Source/Runtime/Experimental/Chaos/Private/GeometryCollection/GeometryCollection.cpp
+++ b/Engine/Source/Runtime/Experimental/Chaos/Private/GeometryCollection/GeometryCollection.cpp
@@ -1242,6 +1242,29 @@ bool FGeometryCollection::SetNumUVLayers(int32 NumLayers)
 	return true;
 }
 
+bool FGeometryCollection::IsVisible(int32 Element) const
+{
+	if (!IsRigid(Element))
+	{
+		return false;
+	}
+
+	if (TransformToGeometryIndex[Element] > INDEX_NONE)
+	{
+		int32 CurrFace = FaceStart[TransformToGeometryIndex[Element]];
+		for (int32 FaceOffset = 0; FaceOffset < FaceCount[TransformToGeometryIndex[Element]]; ++FaceOffset)
+		{
+			if (Visible[CurrFace + FaceOffset])
+			{
+				return true;
+			}
+		}
+	}
+
+	return false;;
+}
+
+
 FGeometryCollection* FGeometryCollection::NewGeometryCollection(const TArray<float>& RawVertexArray, const TArray<int32>& RawIndicesArray, bool ReverseVertexOrder)
 {
 
diff --git a/Engine/Source/Runtime/Experimental/Chaos/Public/GeometryCollection/GeometryCollection.h b/Engine/Source/Runtime/Experimental/Chaos/Public/GeometryCollection/GeometryCollection.h
--- a/Engine/Source/Runtime/Experimental/Chaos/Public/GeometryCollection/GeometryCollection.h
+++ b/Engine/Source/Runtime/Experimental/Chaos/Public/GeometryCollection/GeometryCollection.h
@@ -230,6 +230,9 @@ public:
 	FORCEINLINE void ClearFlags(int32 Element, int32 InFlags) { TManagedArray<int32>& Status = StatusFlags; Status[Element] = Status[Element] & ~InFlags; }
 	FORCEINLINE bool HasFlags(int32 Element, int32 InFlags) const { const TManagedArray<int32>& Status = StatusFlags; return (Status[Element] & InFlags) != 0; }
 
+	/** Return true if the Element contains any visible faces. */
+	bool IsVisible(int32 Element) const;
+
 	/** Connection of leaf geometry */
 	TArray<TArray<int32>> ConnectionGraph();
 
