FractureTools
-- Selection following hierarchy re-topologize removed from Uncluster tool as it causes crashes.

diff --git a/Engine/Plugins/Experimental/ChaosEditor/Source/FractureEditor/Private/FractureToolClustering.cpp b/Engine/Plugins/Experimental/ChaosEditor/Source/FractureEditor/Private/FractureToolClustering.cpp
--- a/Engine/Plugins/Experimental/ChaosEditor/Source/FractureEditor/Private/FractureToolClustering.cpp
+++ b/Engine/Plugins/Experimental/ChaosEditor/Source/FractureEditor/Private/FractureToolClustering.cpp
@@ -174,19 +174,11 @@ void UFractureToolUncluster::Execute(TWeakPtr<FFractureEditorModeToolkit> InTool
 			Context.ConvertSelectionToClusterNodes();
 			Context.RemoveRootNodes();
 
-			// Once the operation is complete, we'll select the children that were re-leveled
-			TArray<int32> NewSelection;
-			for (int32 Cluster : Context.GetSelection())
-			{
-				NewSelection.Append(Children[Cluster].Array());
-			}
-
 			FGeometryCollectionClusteringUtility::CollapseHierarchyOneLevel(Context.GetGeometryCollection().Get(), Context.GetSelection());
-			Context.SetSelection(NewSelection);
 			
 			Refresh(Context, Toolkit);
 		}
-
+		
 		SetOutlinerComponents(Contexts, Toolkit);
 	}
 }
