add tooltips to fracture's merge tiny geo tool, to clarify exactly what values it's filtering on.

diff --git a/Engine/Plugins/Experimental/ChaosEditor/Source/FractureEditor/Private/FractureToolFixTinyGeo.h b/Engine/Plugins/Experimental/ChaosEditor/Source/FractureEditor/Private/FractureToolFixTinyGeo.h
--- a/Engine/Plugins/Experimental/ChaosEditor/Source/FractureEditor/Private/FractureToolFixTinyGeo.h
+++ b/Engine/Plugins/Experimental/ChaosEditor/Source/FractureEditor/Private/FractureToolFixTinyGeo.h
@@ -12,9 +12,11 @@ class FFractureToolContext;
 UENUM()
 enum class EGeometrySelectionMethod
 {
+	// Select by cube root of volume
 	VolumeCubeRoot,
+	// Select by cube root of volume relative to the overall shape's cube root of volume
 	RelativeVolume
-	// TODO: also have a manual selection method?
+	//~ TODO: also have a manual selection method?
 };
 
 /** Settings controlling how convex hulls are generated for geometry collections */
@@ -36,7 +38,8 @@ public:
 	UPROPERTY(EditAnywhere, Category = FilterSettings, meta = (ClampMin = ".00001", UIMin = ".1", UIMax = "10", EditCondition = "SelectionMethod == EGeometrySelectionMethod::VolumeCubeRoot", EditConditionHides))
 	double MinVolumeCubeRoot = 1;
 
-	/** If (cube root of) volume relative to the overall geometry (cube root) volume is less than this, geometry should be merged into neighbors */
+	/** If cube root of volume relative to the overall shape's cube root of volume is less than this, the geometry should be merged into its neighbors.
+	(Note: This is a bit different from the histogram viewer's "Relative Size," which instead shows values relative to the largest rigid bone.) */
 	UPROPERTY(EditAnywhere, Category = FilterSettings, meta = (ClampMin = ".001", UIMax = ".1", ClampMax = "1.0", EditCondition = "SelectionMethod == EGeometrySelectionMethod::RelativeVolume", EditConditionHides))
 	double RelativeVolume = .01;
 
