fix crash on undo fracture selection while gizmo is active

diff --git a/Engine/Plugins/Experimental/ChaosEditor/Source/FractureEditor/Private/FractureToolCutter.cpp b/Engine/Plugins/Experimental/ChaosEditor/Source/FractureEditor/Private/FractureToolCutter.cpp
--- a/Engine/Plugins/Experimental/ChaosEditor/Source/FractureEditor/Private/FractureToolCutter.cpp
+++ b/Engine/Plugins/Experimental/ChaosEditor/Source/FractureEditor/Private/FractureToolCutter.cpp
@@ -48,7 +48,7 @@ void UFractureTransformGizmoSettings::ResetGizmo(bool bResetRotation)
 			{
 				Transform.SetRotation(FQuat::Identity);
 			}
-			TransformGizmo->SetNewGizmoTransform(Transform);
+			TransformGizmo->ReinitializeGizmoTransform(Transform);
 		}
 	}
 }
