modify() the geometry collection component before the geometry collection is fractured, so it updates (more) correctly on undo/redo

diff --git a/Engine/Plugins/Experimental/ChaosEditor/Source/FractureEditor/Public/FractureTool.cpp b/Engine/Plugins/Experimental/ChaosEditor/Source/FractureEditor/Public/FractureTool.cpp
--- a/Engine/Plugins/Experimental/ChaosEditor/Source/FractureEditor/Public/FractureTool.cpp
+++ b/Engine/Plugins/Experimental/ChaosEditor/Source/FractureEditor/Public/FractureTool.cpp
@@ -214,6 +214,7 @@ void UFractureModalTool::Execute(TWeakPtr<FFractureEditorModeToolkit> InToolkit)
 		for (FFractureToolContext& FractureContext : FractureContexts)
 		{
 			FractureContext.GetFracturedGeometryCollection()->Modify();
+			FractureContext.GetGeometryCollectionComponent()->Modify();
 
 			int32 FirstNewGeometryIndex = ExecuteFracture(FractureContext);
 			
