FractureTool
-- deprecated all auto cluster methods other than Voronoi

diff --git a/Engine/Plugins/Experimental/ChaosEditor/Source/FractureEditor/Private/AutoClusterFracture.cpp b/Engine/Plugins/Experimental/ChaosEditor/Source/FractureEditor/Private/AutoClusterFracture.cpp
deleted file mode 100644
diff --git a/Engine/Plugins/Experimental/ChaosEditor/Source/FractureEditor/Private/AutoClusterFracture.h b/Engine/Plugins/Experimental/ChaosEditor/Source/FractureEditor/Private/AutoClusterFracture.h
deleted file mode 100644
diff --git a/Engine/Plugins/Experimental/ChaosEditor/Source/FractureEditor/Private/FractureToolAutoCluster.cpp b/Engine/Plugins/Experimental/ChaosEditor/Source/FractureEditor/Private/FractureToolAutoCluster.cpp
--- a/Engine/Plugins/Experimental/ChaosEditor/Source/FractureEditor/Private/FractureToolAutoCluster.cpp
+++ b/Engine/Plugins/Experimental/ChaosEditor/Source/FractureEditor/Private/FractureToolAutoCluster.cpp
@@ -6,8 +6,6 @@
 #include "FractureEditorCommands.h"
 #include "FractureToolContext.h"
 
-#include "AutoClusterFracture.h"
-
 #include "GeometryCollection/GeometryCollectionComponent.h"
 #include "GeometryCollection/GeometryCollectionClusteringUtility.h"
 #include "GeometryCollection/GeometryCollectionAlgo.h"
@@ -75,44 +73,37 @@ void UFractureToolAutoCluster::Execute(TWeakPtr<FFractureEditorModeToolkit> InTo
 
 			int32 StartTransformCount = GeometryCollection->Transform.Num();
 
-			if (AutoClusterSettings->AutoClusterMode < EFractureAutoClusterMode::Voronoi)
-			{
-				UAutoClusterFractureCommand::ClusterChildBonesOfASingleMesh(Context.GetGeometryCollectionComponent(), AutoClusterSettings->AutoClusterMode, AutoClusterSettings->SiteCount);
-			}
-			else
+			for (const int32 ClusterIndex : Context.GetSelection())
 			{
-				for (const int32 ClusterIndex : Context.GetSelection())
-				{
-					FVoronoiPartitioner VoronoiPartition(GeometryCollection, ClusterIndex);
-					VoronoiPartition.KMeansPartition(AutoClusterSettings->SiteCount);
+				FVoronoiPartitioner VoronoiPartition(GeometryCollection, ClusterIndex);
+				VoronoiPartition.KMeansPartition(AutoClusterSettings->SiteCount);
 
-					if (AutoClusterSettings->bEnforceConnectivity)
-					{
-						GenerateProximityIfNecessary(GeometryCollection);
-						VoronoiPartition.SplitDisconnectedPartitions(GeometryCollection);
-					}
+				if (AutoClusterSettings->bEnforceConnectivity)
+				{
+					GenerateProximityIfNecessary(GeometryCollection);
+					VoronoiPartition.SplitDisconnectedPartitions(GeometryCollection);
+				}
 					
-					int32 PartitionCount = VoronoiPartition.GetPartitionCount();
-					int32 NewClusterIndexStart = GeometryCollection->AddElements(PartitionCount, FGeometryCollection::TransformGroup);
+				int32 PartitionCount = VoronoiPartition.GetPartitionCount();
+				int32 NewClusterIndexStart = GeometryCollection->AddElements(PartitionCount, FGeometryCollection::TransformGroup);
 
-					for (int32 Index = 0; Index < PartitionCount; ++Index)
-					{
+				for (int32 Index = 0; Index < PartitionCount; ++Index)
+				{
 						
-						TArray<int32> NewCluster = VoronoiPartition.GetPartition(Index);
-
-						int32 NewClusterIndex = NewClusterIndexStart + Index;
-						GeometryCollection->Parent[NewClusterIndex] = ClusterIndex;
-						GeometryCollection->Children[ClusterIndex].Add(NewClusterIndex);
-						GeometryCollection->BoneName[NewClusterIndex] = "ClusterBone";
-						GeometryCollection->Children[NewClusterIndex] = TSet<int32>(NewCluster);
-						GeometryCollection->SimulationType[NewClusterIndex] = FGeometryCollection::ESimulationTypes::FST_Clustered;
-						GeometryCollection->Transform[NewClusterIndex] = FTransform::Identity;
-						GeometryCollectionAlgo::ParentTransforms(GeometryCollection, NewClusterIndex, NewCluster);
-					}
-					FGeometryCollectionClusteringUtility::UpdateHierarchyLevelOfChildren(GeometryCollection, ClusterIndex);
-					FGeometryCollectionClusteringUtility::RecursivelyUpdateChildBoneNames(ClusterIndex, GeometryCollection->Children, GeometryCollection->BoneName);
-					FGeometryCollectionClusteringUtility::ValidateResults(GeometryCollection);
+					TArray<int32> NewCluster = VoronoiPartition.GetPartition(Index);
+
+					int32 NewClusterIndex = NewClusterIndexStart + Index;
+					GeometryCollection->Parent[NewClusterIndex] = ClusterIndex;
+					GeometryCollection->Children[ClusterIndex].Add(NewClusterIndex);
+					GeometryCollection->BoneName[NewClusterIndex] = "ClusterBone";
+					GeometryCollection->Children[NewClusterIndex] = TSet<int32>(NewCluster);
+					GeometryCollection->SimulationType[NewClusterIndex] = FGeometryCollection::ESimulationTypes::FST_Clustered;
+					GeometryCollection->Transform[NewClusterIndex] = FTransform::Identity;
+					GeometryCollectionAlgo::ParentTransforms(GeometryCollection, NewClusterIndex, NewCluster);
 				}
+				FGeometryCollectionClusteringUtility::UpdateHierarchyLevelOfChildren(GeometryCollection, ClusterIndex);
+				FGeometryCollectionClusteringUtility::RecursivelyUpdateChildBoneNames(ClusterIndex, GeometryCollection->Children, GeometryCollection->BoneName);
+				FGeometryCollectionClusteringUtility::ValidateResults(GeometryCollection);
 			}
 
 			Context.GenerateGuids(StartTransformCount);
diff --git a/Engine/Plugins/Experimental/ChaosEditor/Source/FractureEditor/Private/FractureToolAutoCluster.h b/Engine/Plugins/Experimental/ChaosEditor/Source/FractureEditor/Private/FractureToolAutoCluster.h
--- a/Engine/Plugins/Experimental/ChaosEditor/Source/FractureEditor/Private/FractureToolAutoCluster.h
+++ b/Engine/Plugins/Experimental/ChaosEditor/Source/FractureEditor/Private/FractureToolAutoCluster.h
@@ -3,11 +3,26 @@
 #pragma once
 
 #include "FractureTool.h"
-#include "AutoClusterFracture.h"
 
 #include "FractureToolAutoCluster.generated.h"
 
 
+UENUM()
+enum class EFractureAutoClusterMode : uint8
+{
+	/** Overlapping bounding box*/
+	BoundingBox UMETA(DisplayName = "Bounding Box"),
+
+	/** GC connectivity */
+	Proximity UMETA(DisplayName = "Proximity"),
+
+	/** Distance */
+	Distance UMETA(DisplayName = "Distance"),
+
+	Voronoi UMETA(DisplayName = "Voronoi"),
+};
+
+
 UCLASS(DisplayName = "Auto Cluster", Category = "FractureTools")
 class UFractureAutoClusterSettings : public UFractureToolSettings
 {
@@ -16,12 +31,11 @@ public:
 
 	UFractureAutoClusterSettings(const FObjectInitializer& ObjInit)
 		: Super(ObjInit)
-		, AutoClusterMode(EFractureAutoClusterMode::BoundingBox)
 		, SiteCount(10)
 	{}
 
-	UPROPERTY(EditAnywhere, Category = AutoCluster, meta = (DisplayName = "Mode"))
-	EFractureAutoClusterMode AutoClusterMode;
+	UPROPERTY(meta = (DeprecatedProperty, DeprecationMessage = "Simplified interface now only supports Voronoi clustering."))
+	EFractureAutoClusterMode AutoClusterMode_DEPRECATED;
 
 	UPROPERTY(EditAnywhere, Category = AutoCluster, meta = (DisplayName = "Cluster Sites", UIMin = "1", UIMax = "5000", ClampMin = "1"))
 	uint32 SiteCount=1;
