Fracture Tool
-- clean up text justifies in interface

diff --git a/Engine/Plugins/Experimental/ChaosEditor/Source/FractureEditor/Private/FractureEditorModeToolkit.cpp b/Engine/Plugins/Experimental/ChaosEditor/Source/FractureEditor/Private/FractureEditorModeToolkit.cpp
--- a/Engine/Plugins/Experimental/ChaosEditor/Source/FractureEditor/Private/FractureEditorModeToolkit.cpp
+++ b/Engine/Plugins/Experimental/ChaosEditor/Source/FractureEditor/Private/FractureEditorModeToolkit.cpp
@@ -104,7 +104,7 @@ FFractureViewSettingsCustomization::CustomizeDetails(IDetailLayoutBuilder& Detai
 	ViewCategory.AddProperty(LevelProperty)
 	.CustomWidget()
 	.NameContent()
-	.HAlign(HAlign_Right)
+	.HAlign(HAlign_Left)
 	[
 		SNew(STextBlock)
 		.TextStyle( &FCoreStyle::Get().GetWidgetStyle<FTextBlockStyle>( "SmallText" ) )
