Fracture Tool
-- Histogram now reflects and informs bone selection rather than filtering outliner view

diff --git a/Engine/Plugins/Experimental/ChaosEditor/Source/FractureEditor/Private/FractureEditorModeToolkit.cpp b/Engine/Plugins/Experimental/ChaosEditor/Source/FractureEditor/Private/FractureEditorModeToolkit.cpp
--- a/Engine/Plugins/Experimental/ChaosEditor/Source/FractureEditor/Private/FractureEditorModeToolkit.cpp
+++ b/Engine/Plugins/Experimental/ChaosEditor/Source/FractureEditor/Private/FractureEditorModeToolkit.cpp
@@ -360,18 +360,6 @@ TSharedRef<SDockTab> FFractureEditorModeToolkit::CreateHierarchyTab(const FSpawn
 				HistogramDetailsView.ToSharedRef()
 			]
 			+ SVerticalBox::Slot()
-			.AutoHeight()
-			[
-				SNew(SButton)
-				//.ButtonStyle(FAppStyle::Get(), "PrimaryButton")
-				.TextStyle(FAppStyle::Get(), "ButtonText")
-				.HAlign(HAlign_Center)
-				.ContentPadding(FMargin(10.f, Padding))
-				.OnClicked(this, &FFractureEditorModeToolkit::ResetHistogramSelection)
-				.IsEnabled(this, &FFractureEditorModeToolkit::CanResetFilter)
-				.Text(LOCTEXT("ResetFilterButton", "Reset Filter"))
-			]
-			+ SVerticalBox::Slot()
 			[
 				SAssignNew(HistogramView, SGeometryCollectionHistogram)
 				.OnBoneSelectionChanged(this, &FFractureEditorModeToolkit::OnHistogramBoneSelectionChanged)
@@ -1032,6 +1020,7 @@ void FFractureEditorModeToolkit::SetOutlinerComponents(const TArray<UGeometryCol
 void FFractureEditorModeToolkit::SetBoneSelection(UGeometryCollectionComponent* InRootComponent, const TArray<int32>& InSelectedBones, bool bClearCurrentSelection)
 {
 	OutlinerView->SetBoneSelection(InRootComponent, InSelectedBones, bClearCurrentSelection);
+	HistogramView->SetBoneSelection(InRootComponent, InSelectedBones, bClearCurrentSelection);
 	
 	if (ActiveTool != nullptr)
 	{
@@ -1075,17 +1064,6 @@ bool FFractureEditorModeToolkit::CanExecuteModal() const
 	return false;
 }
 
-FReply FFractureEditorModeToolkit::ResetHistogramSelection()
-{
-	HistogramView->ClearSelection();
-	return FReply::Handled();
-}
-
-bool FFractureEditorModeToolkit::CanResetFilter() const 
-{
-	return HistogramView->IsSelected();
-}
-
 void FFractureEditorModeToolkit::GetSelectedGeometryCollectionComponents(TSet<UGeometryCollectionComponent*>& GeomCompSelection)
 {
 	USelection* SelectionSet = GEditor->GetSelectedActors();
@@ -1444,6 +1422,7 @@ void FFractureEditorModeToolkit::OnOutlinerBoneSelectionChanged(UGeometryCollect
 
 			FFractureSelectionTools::ToggleSelectedBones(RootComponent, SelectedBones, true, false);
 			OutlinerView->SetBoneSelection(RootComponent, SelectedBones, true);
+			HistogramView->SetBoneSelection(RootComponent, SelectedBones, true);
 		}
 		else
 		{
@@ -1463,8 +1442,33 @@ void FFractureEditorModeToolkit::OnOutlinerBoneSelectionChanged(UGeometryCollect
 
 void FFractureEditorModeToolkit::OnHistogramBoneSelectionChanged(UGeometryCollectionComponent* RootComponent, TArray<int32>& SelectedBones)
 {
-	// If anything is selected in the Histogram, filter the appropriate Outliner component
-	OutlinerView->SetHistogramSelection(RootComponent, SelectedBones);
+	const UGeometryCollection* RestCollection = RootComponent->GetRestCollection();
+	if (RestCollection && !RestCollection->IsPendingKill())
+	{
+		FScopedTransaction Transaction(FractureTransactionContexts::SelectBoneContext, LOCTEXT("SelectGeometryCollectionBoneTransaction", "Select Bone"), RootComponent);
+
+		if (SelectedBones.Num())
+		{
+
+			FFractureSelectionTools::ToggleSelectedBones(RootComponent, SelectedBones, true, false);
+			OutlinerView->SetBoneSelection(RootComponent, SelectedBones, true);
+			HistogramView->SetBoneSelection(RootComponent, SelectedBones, true);
+		}
+		else
+		{
+			FFractureSelectionTools::ClearSelectedBones(RootComponent);
+		}
+
+		if (ActiveTool != nullptr)
+		{
+			ActiveTool->SelectedBonesChanged();
+			ActiveTool->FractureContextChanged();
+		}
+
+		RootComponent->MarkRenderStateDirty();
+		RootComponent->MarkRenderDynamicDataDirty();
+	}
+
 }
 
 FText FFractureEditorModeToolkit::GetSelectionInfo() const
diff --git a/Engine/Plugins/Experimental/ChaosEditor/Source/FractureEditor/Public/FractureEditorModeToolkit.h b/Engine/Plugins/Experimental/ChaosEditor/Source/FractureEditor/Public/FractureEditorModeToolkit.h
--- a/Engine/Plugins/Experimental/ChaosEditor/Source/FractureEditor/Public/FractureEditorModeToolkit.h
+++ b/Engine/Plugins/Experimental/ChaosEditor/Source/FractureEditor/Public/FractureEditorModeToolkit.h
@@ -132,10 +132,6 @@ public:
 	FReply OnModalClicked();
 	bool CanExecuteModal() const;
 
-	// Filter callbacks
-	FReply ResetHistogramSelection();
-	bool CanResetFilter() const;
-
 	static void GetSelectedGeometryCollectionComponents(TSet<UGeometryCollectionComponent*>& GeomCompSelection);
 
 	static void AddAdditionalAttributesIfRequired(UGeometryCollection* GeometryCollectionObject);
