Replace fracture plane cut's "reference actor" concept with a transform gizmo that directly controls the cutting plane

diff --git a/Engine/Plugins/Experimental/ChaosEditor/Source/FractureEditor/Private/FractureToolCutter.h b/Engine/Plugins/Experimental/ChaosEditor/Source/FractureEditor/Private/FractureToolCutter.h
--- a/Engine/Plugins/Experimental/ChaosEditor/Source/FractureEditor/Private/FractureToolCutter.h
+++ b/Engine/Plugins/Experimental/ChaosEditor/Source/FractureEditor/Private/FractureToolCutter.h
@@ -188,7 +188,7 @@ public:
 
 	bool IsGizmoEnabled()
 	{
-		return bUseGizmo;
+		return bUseGizmo && TransformGizmo->IsVisible();
 	}
 
 	FTransform GetTransform()
diff --git a/Engine/Plugins/Experimental/ChaosEditor/Source/FractureEditor/Private/FractureToolPlaneCut.cpp b/Engine/Plugins/Experimental/ChaosEditor/Source/FractureEditor/Private/FractureToolPlaneCut.cpp
--- a/Engine/Plugins/Experimental/ChaosEditor/Source/FractureEditor/Private/FractureToolPlaneCut.cpp
+++ b/Engine/Plugins/Experimental/ChaosEditor/Source/FractureEditor/Private/FractureToolPlaneCut.cpp
@@ -16,8 +16,24 @@ UFractureToolPlaneCut::UFractureToolPlaneCut(const FObjectInitializer& ObjInit)
 {
 	PlaneCutSettings = NewObject<UFracturePlaneCutSettings>(GetTransientPackage(), UFracturePlaneCutSettings::StaticClass());
 	PlaneCutSettings->OwnerTool = this;
+
+	GizmoSettings = NewObject<UFractureTransformGizmoSettings>(GetTransientPackage(), UFractureTransformGizmoSettings::StaticClass());
+	GizmoSettings->OwnerTool = this;
+}
+
+
+void UFractureToolPlaneCut::Setup()
+{
+	GizmoSettings->Setup(this);
 }
 
+
+void UFractureToolPlaneCut::Shutdown()
+{
+	GizmoSettings->Shutdown();
+}
+
+
 FText UFractureToolPlaneCut::GetDisplayText() const
 {
 	return FText(NSLOCTEXT("Fracture", "FractureToolPlaneCut", "Plane Cut Fracture")); 
@@ -44,28 +60,36 @@ void UFractureToolPlaneCut::Render(const FSceneView* View, FViewport* Viewport,
 	const UFracturePlaneCutSettings* LocalCutSettings = PlaneCutSettings;
 	if (CutterSettings->bDrawDiagram)
 	{
-		if (LocalCutSettings->ReferenceActor != nullptr) // so we update with the ref actor realtime
+		// Draw a point centered at plane origin, and a square on the plane around it.
+		auto DrawPlane = [&PDI](const FTransform& Transform, float PlaneSize)
 		{
-			FTransform Transform(LocalCutSettings->ReferenceActor->GetActorTransform());
-			PDI->DrawPoint(Transform.GetLocation(), FLinearColor::Green, 4.f, SDPG_Foreground);
-
-			PDI->DrawLine(Transform.GetLocation(), Transform.GetLocation() + Transform.GetUnitAxis(EAxis::X) * 100.f, FLinearColor(255, 0, 0), SDPG_Foreground);
-			PDI->DrawLine(Transform.GetLocation(), Transform.GetLocation() + Transform.GetUnitAxis(EAxis::Y) * 100.f, FLinearColor(0, 255, 0), SDPG_Foreground);
-
-			PDI->DrawLine(Transform.GetLocation() + Transform.GetUnitAxis(EAxis::X) * 100.f, Transform.GetLocation() + Transform.GetUnitAxis(EAxis::X) * 100.f + Transform.GetUnitAxis(EAxis::Y) * 100.f, FLinearColor(255, 0, 0), SDPG_Foreground);
-			PDI->DrawLine(Transform.GetLocation() + Transform.GetUnitAxis(EAxis::Y) * 100.f, Transform.GetLocation() + Transform.GetUnitAxis(EAxis::X) * 100.f + Transform.GetUnitAxis(EAxis::Y) * 100.f, FLinearColor(0, 255, 0), SDPG_Foreground);
+			FVector Center = Transform.GetLocation();
+			FVector X = (PlaneSize * .5) * Transform.GetUnitAxis(EAxis::X);
+			FVector Y = (PlaneSize * .5) * Transform.GetUnitAxis(EAxis::Y);
+			FVector Corners[4]
+			{
+				Center - X - Y,
+				Center + X - Y,
+				Center + X + Y,
+				Center - X + Y
+			};
+			PDI->DrawPoint(Center, FLinearColor::Green, 4.f, SDPG_Foreground);
+			PDI->DrawLine(Corners[0], Corners[1], FLinearColor(255, 0, 0), SDPG_Foreground);
+			PDI->DrawLine(Corners[1], Corners[2], FLinearColor(0, 255, 0), SDPG_Foreground);
+			PDI->DrawLine(Corners[2], Corners[3], FLinearColor(255, 0, 0), SDPG_Foreground);
+			PDI->DrawLine(Corners[3], Corners[0], FLinearColor(0, 255, 0), SDPG_Foreground);
+		};
+
+		if (GizmoSettings->IsGizmoEnabled())
+		{
+			const FTransform& Transform = GizmoSettings->GetTransform();
+			DrawPlane(Transform, 100.f);
 		}
 		else // draw from computed transforms
 		{
 			for (const FTransform& Transform : RenderCuttingPlanesTransforms)
 			{
-				PDI->DrawPoint(Transform.GetLocation(), FLinearColor::Green, 4.f, SDPG_Foreground);
-
-				PDI->DrawLine(Transform.GetLocation(), Transform.GetLocation() + Transform.GetUnitAxis(EAxis::X) * RenderCuttingPlaneSize, FLinearColor(255, 0, 0), SDPG_Foreground);
-				PDI->DrawLine(Transform.GetLocation(), Transform.GetLocation() + Transform.GetUnitAxis(EAxis::Y) * RenderCuttingPlaneSize, FLinearColor(0, 255, 0), SDPG_Foreground);
-
-				PDI->DrawLine(Transform.GetLocation() + Transform.GetUnitAxis(EAxis::X) * RenderCuttingPlaneSize, Transform.GetLocation() + Transform.GetUnitAxis(EAxis::X) * RenderCuttingPlaneSize + Transform.GetUnitAxis(EAxis::Y) * RenderCuttingPlaneSize, FLinearColor(255, 0, 0), SDPG_Foreground);
-				PDI->DrawLine(Transform.GetLocation() + Transform.GetUnitAxis(EAxis::Y) * RenderCuttingPlaneSize, Transform.GetLocation() + Transform.GetUnitAxis(EAxis::X) * RenderCuttingPlaneSize + Transform.GetUnitAxis(EAxis::Y) * RenderCuttingPlaneSize, FLinearColor(0, 255, 0), SDPG_Foreground);
+				DrawPlane(Transform, RenderCuttingPlaneSize);
 			}
 		}
 	}
@@ -75,6 +99,7 @@ TArray<UObject*> UFractureToolPlaneCut::GetSettingsObjects() const
  {
 	TArray<UObject*> Settings;
 	Settings.Add(CutterSettings);
+	Settings.Add(GizmoSettings);
 	Settings.Add(CollisionSettings);
 	Settings.Add(PlaneCutSettings);
 	return Settings;
@@ -108,9 +133,9 @@ int32 UFractureToolPlaneCut::ExecuteFracture(const FFractureToolContext& Fractur
 		const UFracturePlaneCutSettings* LocalCutSettings = PlaneCutSettings;
 
 		TArray<FPlane> CuttingPlanes;
-		if (LocalCutSettings->ReferenceActor != nullptr)
+		if (GizmoSettings->IsGizmoEnabled())
 		{
-			FTransform Transform(LocalCutSettings->ReferenceActor->GetActorTransform());
+			FTransform Transform = GizmoSettings->GetTransform();
 			CuttingPlanes.Add(FPlane(Transform.GetLocation(), Transform.GetUnitAxis(EAxis::Z)));
 		}
 		else
@@ -158,4 +183,10 @@ void UFractureToolPlaneCut::GenerateSliceTransforms(const FFractureToolContext&
 	}
 }
 
+void UFractureToolPlaneCut::SelectedBonesChanged()
+{
+	GizmoSettings->ResetGizmo();
+	Super::SelectedBonesChanged();
+}
+
 #undef LOCTEXT_NAMESPACE
\ No newline at end of file
diff --git a/Engine/Plugins/Experimental/ChaosEditor/Source/FractureEditor/Private/FractureToolPlaneCut.h b/Engine/Plugins/Experimental/ChaosEditor/Source/FractureEditor/Private/FractureToolPlaneCut.h
--- a/Engine/Plugins/Experimental/ChaosEditor/Source/FractureEditor/Private/FractureToolPlaneCut.h
+++ b/Engine/Plugins/Experimental/ChaosEditor/Source/FractureEditor/Private/FractureToolPlaneCut.h
@@ -19,14 +19,9 @@ public:
 		: Super(ObjInit)
 		, NumberPlanarCuts(1) {}
 
-	/** Number of Clusters - Cluster Voronoi Method */
-	UPROPERTY(EditAnywhere, Category = PlaneCut, meta = (DisplayName = "Number of Cuts", UIMin = "1", UIMax = "20", ClampMin = "1"))
+	/** Number of Cutting Planes */
+	UPROPERTY(EditAnywhere, Category = PlaneCut, meta = (DisplayName = "Number of Cuts", UIMin = "1", UIMax = "20", ClampMin = "1", EditCondition = "bCanCutWithMultiplePlanes"))
 	int32 NumberPlanarCuts;
-
-	/** Actor to be used for voronoi bounds or plane cutting  */
-	UPROPERTY(EditAnywhere, Category = PlaneCut, meta = (DisplayName = "Reference Actor"))
-	TLazyObjectPtr<AActor> ReferenceActor;
-
 };
 
 
@@ -42,6 +37,7 @@ public:
 	virtual FText GetDisplayText() const override;
 	virtual FText GetTooltipText() const override;
 	virtual FSlateIcon GetToolIcon() const override;
+	virtual void SelectedBonesChanged() override;
 
 	void Render(const FSceneView* View, FViewport* Viewport, FPrimitiveDrawInterface* PDI) override;
 
@@ -52,11 +48,17 @@ public:
 	virtual void FractureContextChanged() override;
 	virtual int32 ExecuteFracture(const FFractureToolContext& FractureContext) override;
 
+	virtual void Setup() override;
+	virtual void Shutdown() override;
+
 private:
 	// Slicing
 	UPROPERTY(EditAnywhere, Category = Slicing)
 	TObjectPtr<UFracturePlaneCutSettings> PlaneCutSettings;
 
+	UPROPERTY(EditAnywhere, Category = Uniform)
+	TObjectPtr<UFractureTransformGizmoSettings> GizmoSettings;
+
 	void GenerateSliceTransforms(const FFractureToolContext& Context, TArray<FTransform>& CuttingPlaneTransforms);
 
 	float RenderCuttingPlaneSize;
diff --git a/Engine/Source/Runtime/InteractiveToolsFramework/Public/BaseGizmos/TransformGizmo.h b/Engine/Source/Runtime/InteractiveToolsFramework/Public/BaseGizmos/TransformGizmo.h
--- a/Engine/Source/Runtime/InteractiveToolsFramework/Public/BaseGizmos/TransformGizmo.h
+++ b/Engine/Source/Runtime/InteractiveToolsFramework/Public/BaseGizmos/TransformGizmo.h
@@ -360,6 +360,14 @@ public:
 	 */
 	virtual void SetVisibility(bool bVisible);
 
+	/**
+	 * @return true if Gizmo is visible
+	 */
+	virtual bool IsVisible()
+	{
+		return !GizmoActor->IsHidden();
+	}
+
 	/**
 	 * If true, then when using world frame, Axis and Plane translation snap to the world grid via the ContextQueriesAPI (in PositionSnapFunction)
 	 */
