fix crash if fracture has more than 10 levels

#rnx

diff --git a/Engine/Plugins/Experimental/ChaosEditor/Source/FractureEditor/Private/FractureEditorModeToolkit.cpp b/Engine/Plugins/Experimental/ChaosEditor/Source/FractureEditor/Private/FractureEditorModeToolkit.cpp
--- a/Engine/Plugins/Experimental/ChaosEditor/Source/FractureEditor/Private/FractureEditorModeToolkit.cpp
+++ b/Engine/Plugins/Experimental/ChaosEditor/Source/FractureEditor/Private/FractureEditorModeToolkit.cpp
@@ -1509,7 +1509,6 @@ FText FFractureEditorModeToolkit::GetStatisticsSummary() const
 	if (GeometryCollectionArray.Num() > 0)
 	{
 		TArray<int32> LevelTransformsAll;
-		LevelTransformsAll.SetNumZeroed(10);
 		int32 LevelMax = INT_MIN;
 		int32 EmbeddedCount = 0;
 
@@ -1537,15 +1536,18 @@ FText FFractureEditorModeToolkit::GetStatisticsSummary() const
 					else
 					{
 						const int32 NodeLevel = Levels[Element];
-						while (LevelTransforms.Num() <= NodeLevel)
+						if (LevelTransforms.Num() <= NodeLevel)
 						{
-							LevelTransforms.SetNum(NodeLevel + 1);
-							LevelTransforms[NodeLevel] = 0;
+							LevelTransforms.SetNumZeroed(NodeLevel + 1);
 						}
 						++LevelTransforms[NodeLevel];
 					}		
 				}
 
+				if (LevelTransformsAll.Num() < LevelTransforms.Num())
+				{
+					LevelTransformsAll.SetNumZeroed(LevelTransforms.Num());
+				}
 				for(int32 Level = 0; Level < LevelTransforms.Num(); ++Level)
 				{
 					LevelTransformsAll[Level] += LevelTransforms[Level];
