Fracture Tool
-- fix clipping in histogram interface

diff --git a/Engine/Plugins/Experimental/ChaosEditor/Source/FractureEditor/Private/FractureEditorModeToolkit.cpp b/Engine/Plugins/Experimental/ChaosEditor/Source/FractureEditor/Private/FractureEditorModeToolkit.cpp
--- a/Engine/Plugins/Experimental/ChaosEditor/Source/FractureEditor/Private/FractureEditorModeToolkit.cpp
+++ b/Engine/Plugins/Experimental/ChaosEditor/Source/FractureEditor/Private/FractureEditorModeToolkit.cpp
@@ -351,6 +351,7 @@ TSharedRef<SDockTab> FFractureEditorModeToolkit::CreateHierarchyTab(const FSpawn
 		.BodyBorderBackgroundColor(FLinearColor(1.0, 0.0, 0.0))
 		.AreaTitleFont(FEditorStyle::Get().GetFontStyle("HistogramDetailsView.CategoryFontStyle"))
 		.InitiallyCollapsed(true)
+		.Clipping(EWidgetClipping::ClipToBounds)
 		.BodyContent()
 		[
 			SNew(SVerticalBox)
