fix undo/redo of gizmo movements in fracture tools + stop drawing fracture diagrams for invalid selections

#rnx

diff --git a/Engine/Plugins/Experimental/ChaosEditor/Source/FractureEditor/Private/FractureToolCutter.cpp b/Engine/Plugins/Experimental/ChaosEditor/Source/FractureEditor/Private/FractureToolCutter.cpp
--- a/Engine/Plugins/Experimental/ChaosEditor/Source/FractureEditor/Private/FractureToolCutter.cpp
+++ b/Engine/Plugins/Experimental/ChaosEditor/Source/FractureEditor/Private/FractureToolCutter.cpp
@@ -40,7 +40,7 @@ void UFractureTransformGizmoSettings::ResetGizmo(bool bResetRotation)
 	TransformGizmo->SetVisibility((bool)CombinedBounds.IsValid);
 	if (CombinedBounds.IsValid)
 	{
-		if (bCenterOnSelection)
+		if (bCenterOnSelection && !GIsTransacting)
 		{
 			FTransform Transform = TransformProxy->GetTransform();
 			Transform.SetTranslation(CombinedBounds.GetCenter());
@@ -48,7 +48,7 @@ void UFractureTransformGizmoSettings::ResetGizmo(bool bResetRotation)
 			{
 				Transform.SetRotation(FQuat::Identity);
 			}
-			TransformGizmo->ReinitializeGizmoTransform(Transform);
+			TransformGizmo->SetNewGizmoTransform(Transform);
 		}
 	}
 }
@@ -246,6 +246,10 @@ void UFractureToolVoronoiCutterBase::FractureContextChanged()
 
 	for (FFractureToolContext& FractureContext : FractureContexts)
 	{
+		if (!FractureContext.GetBounds().IsValid) // skip contexts w/ invalid bounds
+		{
+			continue;
+		}
 		// Move the local bounds to the actor so we we'll draw in the correct location
 		GenerateVoronoiSites(FractureContext, VoronoiSites);
 		FBox VoronoiBounds = GetVoronoiBounds(FractureContext, VoronoiSites);
diff --git a/Engine/Plugins/Experimental/ChaosEditor/Source/FractureEditor/Private/FractureToolPlaneCut.cpp b/Engine/Plugins/Experimental/ChaosEditor/Source/FractureEditor/Private/FractureToolPlaneCut.cpp
--- a/Engine/Plugins/Experimental/ChaosEditor/Source/FractureEditor/Private/FractureToolPlaneCut.cpp
+++ b/Engine/Plugins/Experimental/ChaosEditor/Source/FractureEditor/Private/FractureToolPlaneCut.cpp
@@ -117,6 +117,10 @@ void UFractureToolPlaneCut::FractureContextChanged()
 	{
 		// Move the local bounds to the actor so we we'll draw in the correct location
 		FBox Bounds = FractureContext.GetWorldBounds();
+		if (!Bounds.IsValid)
+		{
+			continue;
+		}
 		GenerateSliceTransforms(FractureContext, RenderCuttingPlanesTransforms);
 
 		if (Bounds.GetExtent().GetMax() < RenderCuttingPlaneSize)
