fix voronoi diagrams when multiple contexts are selected (don't accumulate points incorrectly)
 + refuse to compute/render preview diagrams if they become too large, to avoid making the editor non-responsive

#rnx

diff --git a/Engine/Plugins/Experimental/ChaosEditor/Source/FractureEditor/Private/FractureToolCutter.cpp b/Engine/Plugins/Experimental/ChaosEditor/Source/FractureEditor/Private/FractureToolCutter.cpp
--- a/Engine/Plugins/Experimental/ChaosEditor/Source/FractureEditor/Private/FractureToolCutter.cpp
+++ b/Engine/Plugins/Experimental/ChaosEditor/Source/FractureEditor/Private/FractureToolCutter.cpp
@@ -2,6 +2,7 @@
 
 #include "FractureToolCutter.h"
 
+#include "FractureTool.h" // for LogFractureTool
 #include "FractureToolContext.h"
 #include "InteractiveToolsContext.h"
 #include "EditorModeManager.h"
@@ -244,18 +245,43 @@ void UFractureToolVoronoiCutterBase::FractureContextChanged()
 	CellMember.Empty();
 	VoronoiEdges.Empty();
 
+	int32 MaxSitesToShowEdges = 100000; // computing all the voronoi diagrams can make the program non-responsive above this
+	int32 MaxSitesToShowSites = 1000000; // PDI struggles to render the site positions above this
+	bool bEstAboveMaxSites = false;
+
 	for (FFractureToolContext& FractureContext : FractureContexts)
 	{
 		if (!FractureContext.GetBounds().IsValid) // skip contexts w/ invalid bounds
 		{
 			continue;
 		}
-		// Move the local bounds to the actor so we we'll draw in the correct location
-		GenerateVoronoiSites(FractureContext, VoronoiSites);
-		FBox VoronoiBounds = GetVoronoiBounds(FractureContext, VoronoiSites);
-		if (CutterSettings->bDrawDiagram)
+		// Generate voronoi diagrams and cache visualization info
+		TArray<FVector> LocalVoronoiSites;
+		GenerateVoronoiSites(FractureContext, LocalVoronoiSites);
+		// if diagram(s) become too large, skip the visualization
+		if (LocalVoronoiSites.Num() * FractureContexts.Num() > MaxSitesToShowSites || VoronoiSites.Num() + LocalVoronoiSites.Num() > MaxSitesToShowSites)
+		{
+			UE_LOG(LogFractureTool, Warning, TEXT("Voronoi diagram(s) number of sites too large; will not display Voronoi diagram sites"));
+			VoronoiSites.Empty();
+			CellMember.Empty();
+			VoronoiEdges.Empty();
+			break;
+		}
+		VoronoiSites.Append(LocalVoronoiSites);
+		if (bEstAboveMaxSites || LocalVoronoiSites.Num() * FractureContexts.Num() > MaxSitesToShowEdges || VoronoiSites.Num() > MaxSitesToShowEdges)
 		{
-			GetVoronoiEdges(VoronoiSites, VoronoiBounds, VoronoiEdges, CellMember);
+			UE_LOG(LogFractureTool, Warning, TEXT("Voronoi diagram(s) number of sites too large; will not display Voronoi diagram edges"));
+			VoronoiEdges.Empty();
+			CellMember.Empty();
+			bEstAboveMaxSites = true;
+		}
+		else
+		{
+			FBox VoronoiBounds = GetVoronoiBounds(FractureContext, LocalVoronoiSites);
+			if (CutterSettings->bDrawDiagram)
+			{
+				GetVoronoiEdges(LocalVoronoiSites, VoronoiBounds, VoronoiEdges, CellMember);
+			}
 		}
 	}
 }
