FractureTool
-- Fixed dangling cluster generation in several of the tools
-- Removed Cluster slection mode
-- Added Parent, Children, AllInLevel selectin modes
-- Added ClusterMerge tool

diff --git a/Engine/Plugins/Experimental/ChaosEditor/Source/FractureEditor/Private/FractureEditorCommands.h b/Engine/Plugins/Experimental/ChaosEditor/Source/FractureEditor/Private/FractureEditorCommands.h
--- a/Engine/Plugins/Experimental/ChaosEditor/Source/FractureEditor/Private/FractureEditorCommands.h
+++ b/Engine/Plugins/Experimental/ChaosEditor/Source/FractureEditor/Private/FractureEditorCommands.h
@@ -16,8 +16,10 @@ class FFractureEditorCommands : public TCommands<FFractureEditorCommands>
 		TSharedPtr< FUICommandInfo > SelectAll;
 		TSharedPtr< FUICommandInfo > SelectNone;
 		TSharedPtr< FUICommandInfo > SelectNeighbors;
+		TSharedPtr< FUICommandInfo > SelectParent;
+		TSharedPtr< FUICommandInfo > SelectChildren;
 		TSharedPtr< FUICommandInfo > SelectSiblings;
-		TSharedPtr< FUICommandInfo > SelectAllInCluster;
+		TSharedPtr< FUICommandInfo > SelectAllInLevel;
 		TSharedPtr< FUICommandInfo > SelectInvert;
 
 		// View Settings
@@ -34,6 +36,7 @@ class FFractureEditorCommands : public TCommands<FFractureEditorCommands>
 		TSharedPtr< FUICommandInfo > Uncluster;
 		TSharedPtr< FUICommandInfo > Flatten;
 		TSharedPtr< FUICommandInfo > MoveUp;
+		TSharedPtr< FUICommandInfo > ClusterMerge;
 
 		// Edit Commands
 		TSharedPtr< FUICommandInfo > DeleteBranch;
diff --git a/Engine/Plugins/Experimental/ChaosEditor/Source/FractureEditor/Private/FractureEditorModeToolkit.cpp b/Engine/Plugins/Experimental/ChaosEditor/Source/FractureEditor/Private/FractureEditorModeToolkit.cpp
--- a/Engine/Plugins/Experimental/ChaosEditor/Source/FractureEditor/Private/FractureEditorModeToolkit.cpp
+++ b/Engine/Plugins/Experimental/ChaosEditor/Source/FractureEditor/Private/FractureEditorModeToolkit.cpp
@@ -540,8 +540,10 @@ void FFractureEditorModeToolkit::BuildToolPalette(FName PaletteIndex, class FToo
 		ToolbarBuilder.AddToolBarButton(Commands.SelectAll);
 		ToolbarBuilder.AddToolBarButton(Commands.SelectNone);
 		ToolbarBuilder.AddToolBarButton(Commands.SelectNeighbors);
+		ToolbarBuilder.AddToolBarButton(Commands.SelectParent);
+		ToolbarBuilder.AddToolBarButton(Commands.SelectChildren);
 		ToolbarBuilder.AddToolBarButton(Commands.SelectSiblings);
-		ToolbarBuilder.AddToolBarButton(Commands.SelectAllInCluster);
+		ToolbarBuilder.AddToolBarButton(Commands.SelectAllInLevel);
 		ToolbarBuilder.AddToolBarButton(Commands.SelectInvert);
 	}
 	else if (PaletteIndex == TEXT("Fracture"))
@@ -568,6 +570,7 @@ void FFractureEditorModeToolkit::BuildToolPalette(FName PaletteIndex, class FToo
 		ToolbarBuilder.AddToolBarButton(Commands.Cluster);
 		ToolbarBuilder.AddToolBarButton(Commands.Uncluster);
 		ToolbarBuilder.AddToolBarButton(Commands.MoveUp);
+		ToolbarBuilder.AddToolBarButton(Commands.ClusterMerge);
 	}
 	else if (PaletteIndex == TEXT("Embed"))
 	{
diff --git a/Engine/Plugins/Experimental/ChaosEditor/Source/FractureEditor/Private/FractureEditorStyle.cpp b/Engine/Plugins/Experimental/ChaosEditor/Source/FractureEditor/Private/FractureEditorStyle.cpp
--- a/Engine/Plugins/Experimental/ChaosEditor/Source/FractureEditor/Private/FractureEditorStyle.cpp
+++ b/Engine/Plugins/Experimental/ChaosEditor/Source/FractureEditor/Private/FractureEditorStyle.cpp
@@ -34,8 +34,10 @@ FFractureEditorStyle::FFractureEditorStyle()
 	Set("FractureEditor.SelectAll",            new IMAGE_BRUSH_SVG("FractureSelectAll", IconSize));
 	Set("FractureEditor.SelectNone",           new IMAGE_BRUSH_SVG("FractureDeselectAll", IconSize));
 	Set("FractureEditor.SelectNeighbors",      new IMAGE_BRUSH_SVG("FractureSelectNeighbor", IconSize));
+	Set("FractureEditor.SelectParent",		   new IMAGE_BRUSH_SVG("FractureSelectParent", IconSize));
+	Set("FractureEditor.SelectChildren",       new IMAGE_BRUSH_SVG("FractureSelectChildren", IconSize));
 	Set("FractureEditor.SelectSiblings",       new IMAGE_BRUSH_SVG("FractureSelectSiblings", IconSize));
-	Set("FractureEditor.SelectAllInCluster",   new IMAGE_BRUSH_SVG("FractureSelectAllInCluster", IconSize));
+	Set("FractureEditor.SelectAllInLevel",     new IMAGE_BRUSH_SVG("FractureSelectAllInLevel", IconSize));
 	Set("FractureEditor.SelectInvert",         new IMAGE_BRUSH_SVG("FractureSelectInvert", IconSize));
 
 	Set("FractureEditor.AutoCluster",          new IMAGE_BRUSH_SVG("FractureAutoCluster", IconSize));
diff --git a/Engine/Plugins/Experimental/ChaosEditor/Source/FractureEditor/Private/FractureToolClustering.cpp b/Engine/Plugins/Experimental/ChaosEditor/Source/FractureEditor/Private/FractureToolClustering.cpp
--- a/Engine/Plugins/Experimental/ChaosEditor/Source/FractureEditor/Private/FractureToolClustering.cpp
+++ b/Engine/Plugins/Experimental/ChaosEditor/Source/FractureEditor/Private/FractureToolClustering.cpp
@@ -225,4 +225,61 @@ void UFractureToolMoveUp::Execute(TWeakPtr<FFractureEditorModeToolkit> InToolkit
 	}
 }
 
+
+
+FText UFractureToolClusterMerge::GetDisplayText() const
+{
+	return FText(NSLOCTEXT("FractureToolClusteringOps", "FractureToolClusterMerge", "Cluster Merge"));
+}
+
+FText UFractureToolClusterMerge::GetTooltipText() const
+{
+	return FText(NSLOCTEXT("FractureToolClusteringOps", "FractureToolClusterMergeTooltip", "Merge selected clusters."));
+}
+
+FSlateIcon UFractureToolClusterMerge::GetToolIcon() const
+{
+	return FSlateIcon("FractureEditorStyle", "FractureEditor.Merge");
+}
+
+void UFractureToolClusterMerge::RegisterUICommand(FFractureEditorCommands* BindingContext)
+{
+	UI_COMMAND_EXT(BindingContext, UICommandInfo, "Merge", "Merge", "Merge selected clusters.", EUserInterfaceActionType::Button, FInputChord());
+	BindingContext->ClusterMerge = UICommandInfo;
+}
+
+void UFractureToolClusterMerge::Execute(TWeakPtr<FFractureEditorModeToolkit> InToolkit)
+{
+	if (InToolkit.IsValid())
+	{
+		FFractureEditorModeToolkit* Toolkit = InToolkit.Pin().Get();
+
+		TArray<FFractureToolContext> Contexts = GetFractureToolContexts();
+
+		for (FFractureToolContext& Context : Contexts)
+		{
+			const TManagedArray<TSet<int32>>& Children = Context.GetGeometryCollection()->Children;
+
+			Context.ConvertSelectionToClusterNodes();
+
+			// Collect children of context clusters
+			TArray<int32> ChildBones;
+			ChildBones.Reserve(Context.GetGeometryCollection()->NumElements(FGeometryCollection::TransformGroup));
+			for (int32 Select : Context.GetSelection())
+			{
+				ChildBones.Append(Children[Select].Array());
+			}
+			
+			int32 MergeNode = FGeometryCollectionClusteringUtility::PickBestNodeToMergeTo(Context.GetGeometryCollection().Get(), Context.GetSelection());
+			FGeometryCollectionClusteringUtility::ClusterBonesUnderExistingNode(Context.GetGeometryCollection().Get(), MergeNode, ChildBones);
+			FGeometryCollectionClusteringUtility::RemoveDanglingClusters(Context.GetGeometryCollection().Get());
+			
+			Context.SetSelection({MergeNode});
+			Refresh(Context, Toolkit);
+		}
+
+		SetOutlinerComponents(Contexts, Toolkit);
+	}
+}
+
 #undef LOCTEXT_NAMESPACE
diff --git a/Engine/Plugins/Experimental/ChaosEditor/Source/FractureEditor/Private/FractureToolClustering.h b/Engine/Plugins/Experimental/ChaosEditor/Source/FractureEditor/Private/FractureToolClustering.h
--- a/Engine/Plugins/Experimental/ChaosEditor/Source/FractureEditor/Private/FractureToolClustering.h
+++ b/Engine/Plugins/Experimental/ChaosEditor/Source/FractureEditor/Private/FractureToolClustering.h
@@ -73,3 +73,20 @@ public:
 	virtual void RegisterUICommand(FFractureEditorCommands* BindingContext) override;
 	virtual void Execute(TWeakPtr<FFractureEditorModeToolkit> InToolkit) override;
 };
+
+
+UCLASS(DisplayName = "ClusterMerge", Category = "FractureTools")
+class UFractureToolClusterMerge : public UFractureActionTool
+{
+public:
+	GENERATED_BODY()
+
+	UFractureToolClusterMerge(const FObjectInitializer& ObjInit) : Super(ObjInit) {}
+
+	// UFractureActionTool Interface
+	virtual FText GetDisplayText() const override;
+	virtual FText GetTooltipText() const override;
+	virtual FSlateIcon GetToolIcon() const override;
+	virtual void RegisterUICommand(FFractureEditorCommands* BindingContext) override;
+	virtual void Execute(TWeakPtr<FFractureEditorModeToolkit> InToolkit) override;
+};
diff --git a/Engine/Plugins/Experimental/ChaosEditor/Source/FractureEditor/Private/FractureToolSelectors.cpp b/Engine/Plugins/Experimental/ChaosEditor/Source/FractureEditor/Private/FractureToolSelectors.cpp
--- a/Engine/Plugins/Experimental/ChaosEditor/Source/FractureEditor/Private/FractureToolSelectors.cpp
+++ b/Engine/Plugins/Experimental/ChaosEditor/Source/FractureEditor/Private/FractureToolSelectors.cpp
@@ -126,6 +126,69 @@ void UFractureToolSelectNeighbors::Execute(TWeakPtr<FFractureEditorModeToolkit>
 }
 
 
+FText UFractureToolSelectParent::GetDisplayText() const
+{
+	return FText(NSLOCTEXT("Fracture", "FractureToolSelectParent", "Select Parent"));
+}
+
+FText UFractureToolSelectParent::GetTooltipText() const
+{
+	return FText(NSLOCTEXT("Fracture", "FractureToolSelectParentTooltip", "Select clusters containing the currently selected bones."));
+}
+
+FSlateIcon UFractureToolSelectParent::GetToolIcon() const
+{
+	return FSlateIcon("FractureEditorStyle", "FractureEditor.SelectParent");
+}
+
+void UFractureToolSelectParent::RegisterUICommand(FFractureEditorCommands* BindingContext)
+{
+	UI_COMMAND_EXT(BindingContext, UICommandInfo, "SelectParent", "Parent", "Select clusters containing the currently selected bones.", EUserInterfaceActionType::Button, FInputChord());
+	BindingContext->SelectParent = UICommandInfo;
+}
+
+void UFractureToolSelectParent::Execute(TWeakPtr<FFractureEditorModeToolkit> InToolkit)
+{
+	if (InToolkit.IsValid())
+	{
+		SelectByMode(InToolkit.Pin().Get(), GeometryCollection::ESelectionMode::Parent);
+	}
+}
+
+
+
+FText UFractureToolSelectChildren::GetDisplayText() const
+{
+	return FText(NSLOCTEXT("Fracture", "FractureToolSelectChildren", "Select Children"));
+}
+
+FText UFractureToolSelectChildren::GetTooltipText() const
+{
+	return FText(NSLOCTEXT("Fracture", "FractureToolSelectChildrenTooltip", "Select all bones that are immediate children of the currently selected clusters."));
+}
+
+FSlateIcon UFractureToolSelectChildren::GetToolIcon() const
+{
+	return FSlateIcon("FractureEditorStyle", "FractureEditor.SelectChildren");
+}
+
+void UFractureToolSelectChildren::RegisterUICommand(FFractureEditorCommands* BindingContext)
+{
+	UI_COMMAND_EXT(BindingContext, UICommandInfo, "SelectChildren", "Children", "Select all bones that are immediate children of the currently selected clusters.", EUserInterfaceActionType::Button, FInputChord());
+	BindingContext->SelectChildren = UICommandInfo;
+}
+
+void UFractureToolSelectChildren::Execute(TWeakPtr<FFractureEditorModeToolkit> InToolkit)
+{
+	if (InToolkit.IsValid())
+	{
+		SelectByMode(InToolkit.Pin().Get(), GeometryCollection::ESelectionMode::Children);
+	}
+}
+
+
+
+
 FText UFractureToolSelectSiblings::GetDisplayText() const
 {
 	return FText(NSLOCTEXT("Fracture", "FractureToolSelectSiblings", "Select Siblings"));
@@ -133,7 +196,7 @@ FText UFractureToolSelectSiblings::GetDisplayText() const
 
 FText UFractureToolSelectSiblings::GetTooltipText() const
 {
-	return FText(NSLOCTEXT("Fracture", "FractureToolSelectSiblingsTooltip", "Select all bones at the same levels as the currently selected bones."));
+	return FText(NSLOCTEXT("Fracture", "FractureToolSelectSiblingsTooltip", "Select all bones sharing the cluster with currently selected bones."));
 }
 
 FSlateIcon UFractureToolSelectSiblings::GetToolIcon() const
@@ -143,7 +206,7 @@ FSlateIcon UFractureToolSelectSiblings::GetToolIcon() const
 
 void UFractureToolSelectSiblings::RegisterUICommand(FFractureEditorCommands* BindingContext)
 {
-	UI_COMMAND_EXT(BindingContext, UICommandInfo, "SelectSiblings", "Siblings", "Select all bones at the same levels as the currently selected bones.", EUserInterfaceActionType::Button, FInputChord());
+	UI_COMMAND_EXT(BindingContext, UICommandInfo, "SelectSiblings", "Siblings", "Select all bones sharing the cluster with currently selected bones.", EUserInterfaceActionType::Button, FInputChord());
 	BindingContext->SelectSiblings = UICommandInfo;
 }
 
@@ -156,36 +219,37 @@ void UFractureToolSelectSiblings::Execute(TWeakPtr<FFractureEditorModeToolkit> I
 }
 
 
-FText UFractureToolSelectAllInCluster::GetDisplayText() const
+FText UFractureToolSelectAllInLevel::GetDisplayText() const
 {
-	return FText(NSLOCTEXT("Fracture", "FractureToolSelectAllInCluster", "Select All In Cluster"));
+	return FText(NSLOCTEXT("Fracture", "FractureToolSelectAllInLevel", "Select All In Level"));
 }
 
-FText UFractureToolSelectAllInCluster::GetTooltipText() const
+FText UFractureToolSelectAllInLevel::GetTooltipText() const
 {
-	return FText(NSLOCTEXT("Fracture", "FractureToolSelectAllInClusterTooltip", "Select all bones with the same parent as selected bones."));
+	return FText(NSLOCTEXT("Fracture", "FractureToolSelectAllInLevelTooltip", "Select all bones at the same level as currently selected bones."));
 }
 
-FSlateIcon UFractureToolSelectAllInCluster::GetToolIcon() const
+FSlateIcon UFractureToolSelectAllInLevel::GetToolIcon() const
 {
-	return FSlateIcon("FractureEditorStyle", "FractureEditor.SelectAllInCluster");
+	return FSlateIcon("FractureEditorStyle", "FractureEditor.SelectAllInLevel");
 }
 
-void UFractureToolSelectAllInCluster::RegisterUICommand(FFractureEditorCommands* BindingContext)
+void UFractureToolSelectAllInLevel::RegisterUICommand(FFractureEditorCommands* BindingContext)
 {
-	UI_COMMAND_EXT(BindingContext, UICommandInfo, "SelectAllInCluster", "Cluster", "Select all bones with the same parent as selected bones.", EUserInterfaceActionType::Button, FInputChord());
-	BindingContext->SelectAllInCluster = UICommandInfo;
+	UI_COMMAND_EXT(BindingContext, UICommandInfo, "SelectAllInLevel", "Level", "Select all bones at the same level as currently selected bones.", EUserInterfaceActionType::Button, FInputChord());
+	BindingContext->SelectAllInLevel = UICommandInfo;
 }
 
-void UFractureToolSelectAllInCluster::Execute(TWeakPtr<FFractureEditorModeToolkit> InToolkit)
+void UFractureToolSelectAllInLevel::Execute(TWeakPtr<FFractureEditorModeToolkit> InToolkit)
 {
 	if (InToolkit.IsValid())
 	{
-		SelectByMode(InToolkit.Pin().Get(), GeometryCollection::ESelectionMode::AllInCluster);
+		SelectByMode(InToolkit.Pin().Get(), GeometryCollection::ESelectionMode::Level);
 	}
 }
 
 
+
 FText UFractureToolSelectInvert::GetDisplayText() const
 {
 	return FText(NSLOCTEXT("Fracture", "FractureToolSelectInvert", "Invert Selection"));
diff --git a/Engine/Plugins/Experimental/ChaosEditor/Source/FractureEditor/Private/FractureToolSelectors.h b/Engine/Plugins/Experimental/ChaosEditor/Source/FractureEditor/Private/FractureToolSelectors.h
--- a/Engine/Plugins/Experimental/ChaosEditor/Source/FractureEditor/Private/FractureToolSelectors.h
+++ b/Engine/Plugins/Experimental/ChaosEditor/Source/FractureEditor/Private/FractureToolSelectors.h
@@ -69,6 +69,39 @@ public:
 };
 
 
+UCLASS(DisplayName = "Select Parent", Category = "FractureTools")
+class UFractureToolSelectParent : public UFractureToolSelectAll
+{
+public:
+	GENERATED_BODY()
+
+	UFractureToolSelectParent(const FObjectInitializer& ObjInit) : Super(ObjInit) {}
+
+	// UFractureActionTool Interface
+	virtual FText GetDisplayText() const override;
+	virtual FText GetTooltipText() const override;
+	virtual FSlateIcon GetToolIcon() const override;
+	virtual void RegisterUICommand(FFractureEditorCommands* BindingContext) override;
+	virtual void Execute(TWeakPtr<FFractureEditorModeToolkit> InToolkit) override;
+};
+
+
+UCLASS(DisplayName = "Select Children", Category = "FractureTools")
+class UFractureToolSelectChildren : public UFractureToolSelectAll
+{
+public:
+	GENERATED_BODY()
+
+	UFractureToolSelectChildren(const FObjectInitializer& ObjInit) : Super(ObjInit) {}
+
+	// UFractureActionTool Interface
+	virtual FText GetDisplayText() const override;
+	virtual FText GetTooltipText() const override;
+	virtual FSlateIcon GetToolIcon() const override;
+	virtual void RegisterUICommand(FFractureEditorCommands* BindingContext) override;
+	virtual void Execute(TWeakPtr<FFractureEditorModeToolkit> InToolkit) override;
+};
+
 UCLASS(DisplayName = "Select Siblings", Category = "FractureTools")
 class UFractureToolSelectSiblings : public UFractureToolSelectAll
 {
@@ -86,13 +119,13 @@ public:
 };
 
 
-UCLASS(DisplayName = "Select Neighbors", Category = "FractureTools")
-class UFractureToolSelectAllInCluster : public UFractureToolSelectAll
+UCLASS(DisplayName = "Select All In Level", Category = "FractureTools")
+class UFractureToolSelectAllInLevel : public UFractureToolSelectAll
 {
 public:
 	GENERATED_BODY()
 
-	UFractureToolSelectAllInCluster(const FObjectInitializer& ObjInit) : Super(ObjInit) {}
+	UFractureToolSelectAllInLevel(const FObjectInitializer& ObjInit) : Super(ObjInit) {}
 
 	// UFractureActionTool Interface
 	virtual FText GetDisplayText() const override;
@@ -103,7 +136,7 @@ public:
 };
 
 
-UCLASS(DisplayName = "Select Neighbors", Category = "FractureTools")
+UCLASS(DisplayName = "Select Invert", Category = "FractureTools")
 class UFractureToolSelectInvert : public UFractureToolSelectAll
 {
 public:
diff --git a/Engine/Source/Runtime/Experimental/Chaos/Private/GeometryCollection/GeometryCollectionClusteringUtility.cpp b/Engine/Source/Runtime/Experimental/Chaos/Private/GeometryCollection/GeometryCollectionClusteringUtility.cpp
--- a/Engine/Source/Runtime/Experimental/Chaos/Private/GeometryCollection/GeometryCollectionClusteringUtility.cpp
+++ b/Engine/Source/Runtime/Experimental/Chaos/Private/GeometryCollection/GeometryCollectionClusteringUtility.cpp
@@ -714,6 +714,8 @@ void FGeometryCollectionClusteringUtility::ValidateResults(FGeometryCollection*
 	}
 	check(NumRootNodes == 1);
 
+	RemoveDanglingClusters(GeometryCollection);
+
 	ensure(GeometryCollection->HasContiguousFaces());
 	ensure(GeometryCollection->HasContiguousVertices());
 }
@@ -813,6 +815,7 @@ void FGeometryCollectionClusteringUtility::MoveUpOneHierarchyLevel(FGeometryColl
 			}
 		}
 	}
+	
 	ValidateResults(GeometryCollection);
 }
 
diff --git a/Engine/Source/Runtime/Experimental/GeometryCollectionEngine/Private/GeometryCollection/GeometryCollectionComponent.cpp b/Engine/Source/Runtime/Experimental/GeometryCollectionEngine/Private/GeometryCollection/GeometryCollectionComponent.cpp
--- a/Engine/Source/Runtime/Experimental/GeometryCollectionEngine/Private/GeometryCollection/GeometryCollectionComponent.cpp
+++ b/Engine/Source/Runtime/Experimental/GeometryCollectionEngine/Private/GeometryCollection/GeometryCollectionComponent.cpp
@@ -2308,11 +2308,10 @@ void FScopedColorEdit::SelectBones(GeometryCollection::ESelectionMode SelectionM
 		}
 		break;
 
-		case GeometryCollection::ESelectionMode::Siblings:
+		case GeometryCollection::ESelectionMode::Parent:
 		{
 			const TManagedArray<int32>& Parents = GeometryCollectionPtr->Parent;
-			const TManagedArray<TSet<int32>>& Children = GeometryCollectionPtr->Children;
-
+			
 			const TArray<int32> SelectedBones = GetSelectedBones();
 
 			TArray<int32> NewSelection;
@@ -2321,12 +2320,28 @@ void FScopedColorEdit::SelectBones(GeometryCollection::ESelectionMode SelectionM
 				int32 ParentBone = Parents[Bone];
 				if (ParentBone != FGeometryCollection::Invalid)
 				{
-					for (int32 Child : Children[ParentBone])
-					{
-						NewSelection.AddUnique(Child);
-					}
+					NewSelection.AddUnique(ParentBone);
 				}
+			}
+
+			ResetBoneSelection();
+			AppendSelectedBones(NewSelection);
+		}
+		break;
 
+		case GeometryCollection::ESelectionMode::Children:
+		{
+			const TManagedArray<TSet<int32>>& Children = GeometryCollectionPtr->Children;
+
+			const TArray<int32> SelectedBones = GetSelectedBones();
+
+			TArray<int32> NewSelection;
+			for (int32 Bone : SelectedBones)
+			{
+				for (int32 Child : Children[Bone])
+				{
+					NewSelection.AddUnique(Child);
+				}
 			}
 
 			ResetBoneSelection();
@@ -2334,9 +2349,10 @@ void FScopedColorEdit::SelectBones(GeometryCollection::ESelectionMode SelectionM
 		}
 		break;
 
-		case GeometryCollection::ESelectionMode::AllInCluster:
+		case GeometryCollection::ESelectionMode::Siblings:
 		{
 			const TManagedArray<int32>& Parents = GeometryCollectionPtr->Parent;
+			const TManagedArray<TSet<int32>>& Children = GeometryCollectionPtr->Children;
 
 			const TArray<int32> SelectedBones = GetSelectedBones();
 
@@ -2344,12 +2360,12 @@ void FScopedColorEdit::SelectBones(GeometryCollection::ESelectionMode SelectionM
 			for (int32 Bone : SelectedBones)
 			{
 				int32 ParentBone = Parents[Bone];
-				TArray<int32> LeafBones;
-				FGeometryCollectionClusteringUtility::GetLeafBones(GeometryCollectionPtr.Get(), ParentBone, true, LeafBones);
-
-				for (int32 Element : LeafBones)
+				if (ParentBone != FGeometryCollection::Invalid)
 				{
-					NewSelection.AddUnique(Element);
+					for (int32 Child : Children[ParentBone])
+					{
+						NewSelection.AddUnique(Child);
+					}
 				}
 
 			}
@@ -2359,6 +2375,33 @@ void FScopedColorEdit::SelectBones(GeometryCollection::ESelectionMode SelectionM
 		}
 		break;
 
+		case GeometryCollection::ESelectionMode::Level:
+		{
+			if (GeometryCollectionPtr->HasAttribute("Level", FTransformCollection::TransformGroup))
+			{
+				const TManagedArray<int32>& Levels = GeometryCollectionPtr->GetAttribute<int32>("Level", FTransformCollection::TransformGroup);
+
+				const TArray<int32> SelectedBones = GetSelectedBones();
+
+				TArray<int32> NewSelection;
+				for (int32 Bone : SelectedBones)
+				{
+					int32 Level = Levels[Bone];
+					for (int32 TransformIdx = 0; TransformIdx < GeometryCollectionPtr->NumElements(FTransformCollection::TransformGroup); ++TransformIdx)
+					{
+						if (Levels[TransformIdx] == Level)
+						{
+							NewSelection.AddUnique(TransformIdx);
+						}
+					}
+				}
+
+				ResetBoneSelection();
+				AppendSelectedBones(NewSelection);
+			}	
+		}
+		break;
+
 		default: 
 			check(false); // unexpected selection mode
 		break;
diff --git a/Engine/Source/Runtime/Experimental/GeometryCollectionEngine/Public/GeometryCollection/GeometryCollectionComponent.h b/Engine/Source/Runtime/Experimental/GeometryCollectionEngine/Public/GeometryCollection/GeometryCollectionComponent.h
--- a/Engine/Source/Runtime/Experimental/GeometryCollectionEngine/Public/GeometryCollection/GeometryCollectionComponent.h
+++ b/Engine/Source/Runtime/Experimental/GeometryCollectionEngine/Public/GeometryCollection/GeometryCollectionComponent.h
@@ -47,7 +47,9 @@ namespace GeometryCollection
 		InverseGeometry,
 		Siblings,
 		Neighbors,
-		AllInCluster
+		Parent,
+		Children,
+		Level
 	};
 }
 
