FractureTool
-- Corrected problematic topology issues with dangling cluster removal

diff --git a/Engine/Plugins/Experimental/ChaosEditor/Source/FractureEditor/Private/FractureToolClustering.cpp b/Engine/Plugins/Experimental/ChaosEditor/Source/FractureEditor/Private/FractureToolClustering.cpp
--- a/Engine/Plugins/Experimental/ChaosEditor/Source/FractureEditor/Private/FractureToolClustering.cpp
+++ b/Engine/Plugins/Experimental/ChaosEditor/Source/FractureEditor/Private/FractureToolClustering.cpp
@@ -51,13 +51,12 @@ void UFractureToolFlattenAll::Execute(TWeakPtr<FFractureEditorModeToolkit> InToo
 				TArray<int32> LeafBones;
 				FGeometryCollectionClusteringUtility::GetLeafBones(Context.GetGeometryCollection().Get(), ClusterIndex, true, LeafBones);
 				FGeometryCollectionClusteringUtility::ClusterBonesUnderExistingNode(Context.GetGeometryCollection().Get(), ClusterIndex, LeafBones);
-
-				// Cleanup: Remove any clusters remaining in the flattened branch.
-				FGeometryCollectionClusteringUtility::RemoveDanglingClusters(Context.GetGeometryCollection().Get());
-
 			}
 
-			Refresh(Context, Toolkit);
+			// Cleanup: Remove any clusters remaining in the flattened branch.
+			FGeometryCollectionClusteringUtility::RemoveDanglingClusters(Context.GetGeometryCollection().Get());
+
+			Refresh(Context, Toolkit, true);
 		}
 
 		SetOutlinerComponents(Contexts, Toolkit);
@@ -176,7 +175,8 @@ void UFractureToolUncluster::Execute(TWeakPtr<FFractureEditorModeToolkit> InTool
 
 			FGeometryCollectionClusteringUtility::CollapseHierarchyOneLevel(Context.GetGeometryCollection().Get(), Context.GetSelection());
 			
-			Refresh(Context, Toolkit);
+			FGeometryCollectionClusteringUtility::RemoveDanglingClusters(Context.GetGeometryCollection().Get());
+			Refresh(Context, Toolkit, true);
 		}
 		
 		SetOutlinerComponents(Contexts, Toolkit);
@@ -218,7 +218,8 @@ void UFractureToolMoveUp::Execute(TWeakPtr<FFractureEditorModeToolkit> InToolkit
 		{
 			Context.ConvertSelectionToRigidNodes();
 			FGeometryCollectionClusteringUtility::MoveUpOneHierarchyLevel(Context.GetGeometryCollection().Get(), Context.GetSelection());
-			Refresh(Context, Toolkit);
+			FGeometryCollectionClusteringUtility::RemoveDanglingClusters(Context.GetGeometryCollection().Get());
+			Refresh(Context, Toolkit, true);
 		}
 		
 		SetOutlinerComponents(Contexts, Toolkit);
diff --git a/Engine/Plugins/Experimental/ChaosEditor/Source/FractureEditor/Private/GeometryCollectionOutlinerDragDrop.cpp b/Engine/Plugins/Experimental/ChaosEditor/Source/FractureEditor/Private/GeometryCollectionOutlinerDragDrop.cpp
--- a/Engine/Plugins/Experimental/ChaosEditor/Source/FractureEditor/Private/GeometryCollectionOutlinerDragDrop.cpp
+++ b/Engine/Plugins/Experimental/ChaosEditor/Source/FractureEditor/Private/GeometryCollectionOutlinerDragDrop.cpp
@@ -121,6 +121,7 @@ bool FGeometryCollectionBoneDragDrop::ReparentBones(const FGeometryCollection* O
 		}
 
 		FGeometryCollectionClusteringUtility::ClusterBonesUnderExistingNode(GeometryCollection.Get(), OtherBone, BonePayload);
+		FGeometryCollectionClusteringUtility::RemoveDanglingClusters(GeometryCollection.Get());
 	}
 
 	return bValid;
diff --git a/Engine/Source/Runtime/Experimental/Chaos/Private/GeometryCollection/GeometryCollectionClusteringUtility.cpp b/Engine/Source/Runtime/Experimental/Chaos/Private/GeometryCollection/GeometryCollectionClusteringUtility.cpp
--- a/Engine/Source/Runtime/Experimental/Chaos/Private/GeometryCollection/GeometryCollectionClusteringUtility.cpp
+++ b/Engine/Source/Runtime/Experimental/Chaos/Private/GeometryCollection/GeometryCollectionClusteringUtility.cpp
@@ -714,8 +714,6 @@ void FGeometryCollectionClusteringUtility::ValidateResults(FGeometryCollection*
 	}
 	check(NumRootNodes == 1);
 
-	RemoveDanglingClusters(GeometryCollection);
-
 	ensure(GeometryCollection->HasContiguousFaces());
 	ensure(GeometryCollection->HasContiguousVertices());
 }
