track fracture tools' explode view updates to allow visualization updates

#rnx

diff --git a/Engine/Plugins/Experimental/ChaosEditor/Source/FractureEditor/Private/FractureToolFixTinyGeo.cpp b/Engine/Plugins/Experimental/ChaosEditor/Source/FractureEditor/Private/FractureToolFixTinyGeo.cpp
--- a/Engine/Plugins/Experimental/ChaosEditor/Source/FractureEditor/Private/FractureToolFixTinyGeo.cpp
+++ b/Engine/Plugins/Experimental/ChaosEditor/Source/FractureEditor/Private/FractureToolFixTinyGeo.cpp
@@ -40,12 +40,45 @@ void UFractureToolFixTinyGeo::RegisterUICommand( FFractureEditorCommands* Bindin
 	BindingContext->FixTinyGeo = UICommandInfo;
 }
 
+const TManagedArray<FVector3f>* UFractureToolFixTinyGeo::GetExplodedVectors(const UGeometryCollection* GeometryCollectionObject)
+{
+	if (GeometryCollectionObject)
+	{
+		const FGeometryCollection& Collection = *GeometryCollectionObject->GetGeometryCollection();
+		if (Collection.HasAttribute("ExplodedVector", FGeometryCollection::TransformGroup))
+		{
+			return &Collection.GetAttribute<FVector3f>("ExplodedVector", FGeometryCollection::TransformGroup);
+		}
+	}
+	return nullptr;
+}
+
 void UFractureToolFixTinyGeo::Render(const FSceneView* View, FViewport* Viewport, FPrimitiveDrawInterface* PDI)
 {
-	for (FBox& Box : ToRemoveBounds)
+	if (!ensure(ToRemoveCollectionBones.Num() == ToRemoveBounds.Num()))
 	{
-		FVector B000 = Box.Min;
-		FVector B111 = Box.Max;
+		return;
+	}
+	const TManagedArray<FVector3f>* ExplodedVectors = nullptr;
+
+	int32 LastCollectionIdx = INDEX_NONE;
+	for (int32 Idx = 0; Idx < ToRemoveBounds.Num(); Idx++)
+	{
+		TPair<int32, int32> CollectBones = ToRemoveCollectionBones[Idx];
+		if (CollectBones.Key != LastCollectionIdx)
+		{
+			ExplodedVectors = GetExplodedVectors(VisCollections[CollectBones.Key]);
+			LastCollectionIdx = CollectBones.Key;
+		}
+		int32 BoneIdx = CollectBones.Value;
+		const FBox& Box = ToRemoveBounds[Idx];
+		FVector Offset = FVector::ZeroVector;
+		if (ExplodedVectors && BoneIdx < ExplodedVectors->Num())
+		{
+			Offset = (*ExplodedVectors)[BoneIdx];
+		}
+		FVector B000 = Box.Min + Offset;
+		FVector B111 = Box.Max + Offset;
 		FVector B011(B000.X, B111.Y, B111.Z);
 		FVector B101(B111.X, B000.Y, B111.Z);
 		FVector B110(B111.X, B111.Y, B000.Z);
@@ -82,11 +115,14 @@ void UFractureToolFixTinyGeo::FractureContextChanged()
 	UpdateDefaultRandomSeed();
 	TArray<FFractureToolContext> FractureContexts = GetFractureToolContexts();
 
+	ToRemoveCollectionBones.Reset();
 	ToRemoveBounds.Reset();
+	VisCollections.Reset();
 
 	for (const FFractureToolContext& FractureContext : FractureContexts)
 	{
 		FGeometryCollection& Collection = *FractureContext.GetGeometryCollection();
+		int CollectionIdx = VisCollections.Emplace(FractureContext.GetGeometryCollectionComponent()->GetRestCollection());
 
 		TArray<double> Volumes;
 		FindBoneVolumes(
@@ -111,11 +147,7 @@ void UFractureToolFixTinyGeo::FractureContextChanged()
 		for (int32 TransformIdx : SmallIndices) // small transforms
 		{
 			FTransform InnerTransform = GeometryCollectionAlgo::GlobalMatrix(Collection.Transform, Collection.Parent, TransformIdx);
-			if (Collection.HasAttribute("ExplodedVector", FGeometryCollection::TransformGroup))
-			{
-				TManagedArray<FVector3f>& ExplodedVectors = Collection.GetAttribute<FVector3f>("ExplodedVector", FGeometryCollection::TransformGroup);
-				InnerTransform = InnerTransform * FTransform(ExplodedVectors[TransformIdx]);
-			}
+
 
 			FTransform CombinedTransform = InnerTransform * OuterTransform;
 
@@ -132,6 +164,7 @@ void UFractureToolFixTinyGeo::FractureContextChanged()
 			{
 				Bounds += CombinedTransform.TransformPosition(Collection.Vertex[VIdx]);
 			}
+			ToRemoveCollectionBones.Add(TPair<int32, int32>(CollectionIdx, TransformIdx));
 			ToRemoveBounds.Add(Bounds);
 		}
 	}
diff --git a/Engine/Plugins/Experimental/ChaosEditor/Source/FractureEditor/Private/FractureToolFixTinyGeo.h b/Engine/Plugins/Experimental/ChaosEditor/Source/FractureEditor/Private/FractureToolFixTinyGeo.h
--- a/Engine/Plugins/Experimental/ChaosEditor/Source/FractureEditor/Private/FractureToolFixTinyGeo.h
+++ b/Engine/Plugins/Experimental/ChaosEditor/Source/FractureEditor/Private/FractureToolFixTinyGeo.h
@@ -77,8 +77,11 @@ private:
 	UPROPERTY(EditAnywhere, Category = FixGeo)
 	TObjectPtr<UFractureTinyGeoSettings> TinyGeoSettings;
 
+	TArray<TPair<int32, int32>> ToRemoveCollectionBones; // Pairs of [index into VisCollections, bone index]
+	TArray<FBox> ToRemoveBounds; // Bounds in global space but without exploded vectors applied
+	TArray<TObjectPtr<const UGeometryCollection>> VisCollections; // source geometry collections
 
-	TArray<FBox> ToRemoveBounds;
+	const TManagedArray<FVector3f>* GetExplodedVectors(const UGeometryCollection* GeometryCollectionObject);
 
 	double GetMinVolume(TArray<double>& Volumes);
 	const double VolDimScale = .01; // compute volumes in meters instead of cm, for saner units at typical scales
