make fracture tool aware of component transform changes, so e.g. the visualization can be updated accordingly

#rnx

diff --git a/Engine/Plugins/Experimental/ChaosEditor/Source/FractureEditor/Public/FractureTool.cpp b/Engine/Plugins/Experimental/ChaosEditor/Source/FractureEditor/Public/FractureTool.cpp
--- a/Engine/Plugins/Experimental/ChaosEditor/Source/FractureEditor/Public/FractureTool.cpp
+++ b/Engine/Plugins/Experimental/ChaosEditor/Source/FractureEditor/Public/FractureTool.cpp
@@ -283,3 +283,16 @@ void UFractureModalTool::PostEditChangeChainProperty(struct FPropertyChangedChai
 	FractureContextChanged();
 }
 
+
+void UFractureModalTool::OnComponentTransformChangedInternal(USceneComponent* InRootComponent, ETeleportType Teleport)
+{
+	if (UGeometryCollectionComponent* GeomColl = Cast<UGeometryCollectionComponent>(InRootComponent))
+	{
+		bool bIsSel = GeomColl->IsSelected();
+		bool bActorIsSel = GeomColl->GetOwner()->IsSelected();
+		if (bIsSel || bActorIsSel)
+		{
+			OnComponentTransformChanged(GeomColl);
+		}
+	}
+}
diff --git a/Engine/Plugins/Experimental/ChaosEditor/Source/FractureEditor/Public/FractureTool.h b/Engine/Plugins/Experimental/ChaosEditor/Source/FractureEditor/Public/FractureTool.h
--- a/Engine/Plugins/Experimental/ChaosEditor/Source/FractureEditor/Public/FractureTool.h
+++ b/Engine/Plugins/Experimental/ChaosEditor/Source/FractureEditor/Public/FractureTool.h
@@ -114,10 +114,26 @@ public:
 	virtual void SelectedBonesChanged() {}
 
 	// Called when the modal tool is entered
-	virtual void Setup() {}
+	virtual void Setup()
+	{
+		GEngine->OnComponentTransformChanged().AddUObject(this, &UFractureModalTool::OnComponentTransformChangedInternal);
+	}
 
 	// Called when the modal tool is exited (on switching to a new modal tool or exiting the fracture editor mode)
-	virtual void Shutdown() {}
+	virtual void Shutdown() 
+	{
+		GEngine->OnComponentTransformChanged().RemoveAll(this);
+	}
+
+	// Called when a selected geometry collection component is moved in the scene
+	virtual void OnComponentTransformChanged(UGeometryCollectionComponent* Component)
+	{
+		FractureContextChanged();
+	}
+
+protected:
+	
+	void OnComponentTransformChangedInternal(USceneComponent* InRootComponent, ETeleportType Teleport);
 
 };
 
