Mode UI layer will set the icon if one is provided to the minor tab info, or fallback on the default for the workspace category if one was not provided.

diff --git a/Engine/Plugins/Experimental/ChaosEditor/Source/FractureEditor/Private/FractureEditorModeToolkit.cpp b/Engine/Plugins/Experimental/ChaosEditor/Source/FractureEditor/Private/FractureEditorModeToolkit.cpp
--- a/Engine/Plugins/Experimental/ChaosEditor/Source/FractureEditor/Private/FractureEditorModeToolkit.cpp
+++ b/Engine/Plugins/Experimental/ChaosEditor/Source/FractureEditor/Private/FractureEditorModeToolkit.cpp
@@ -316,12 +316,14 @@ void FFractureEditorModeToolkit::RequestModeUITabs()
 		HierarchyTabInfo.OnSpawnTab = FOnSpawnTab::CreateSP(SharedThis(this), &FFractureEditorModeToolkit::CreateHierarchyTab);
 		HierarchyTabInfo.TabLabel = LOCTEXT("FractureHierarchy", "Fracture Hierarchy");
 		HierarchyTabInfo.TabTooltip = LOCTEXT("ModesToolboxTabTooltipText", "Open the  Modes tab, which contains the active editor mode's settings.");
+		HierarchyTabInfo.TabIcon = GetEditorModeIcon();
 		ModeUILayerPtr->SetModePanelInfo(FAssetEditorModeUILayer::TopRightTabID, HierarchyTabInfo);
 
 
 		StatisticsTabInfo.OnSpawnTab = FOnSpawnTab::CreateSP(SharedThis(this), &FFractureEditorModeToolkit::CreateStatisticsTab);
 		StatisticsTabInfo.TabLabel = LOCTEXT("FractureStatistics", "Level Statistics");
 		StatisticsTabInfo.TabTooltip = LOCTEXT("ModesToolboxTabTooltipText", "Open the  Modes tab, which contains the active editor mode's settings.");
+		StatisticsTabInfo.TabIcon = GetEditorModeIcon();
 		ModeUILayerPtr->SetModePanelInfo(FAssetEditorModeUILayer::BottomLeftTabID, StatisticsTabInfo);
 	}
 }
@@ -426,11 +428,6 @@ TSharedRef<SDockTab> FFractureEditorModeToolkit::CreateHierarchyTab(const FSpawn
 			OutlinerExpander
 		];
 
-	const FSlateBrush* TabIcon = GetEditorModeIcon().GetSmallIcon();
-	if (CreatedTab)
-	{
-		CreatedTab->SetTabIcon(TabIcon);
-	}
 	HierarchyTab = CreatedTab;
 	return CreatedTab.ToSharedRef();
 }
@@ -456,11 +453,6 @@ TSharedRef<SDockTab> FFractureEditorModeToolkit::CreateStatisticsTab(const FSpaw
 			StatisticsExpander
 		];
 
-	const FSlateBrush* TabIcon = GetEditorModeIcon().GetSmallIcon();
-	if (CreatedTab)
-	{
-		CreatedTab->SetTabIcon(TabIcon);
-	}
 	StatisticsTab = CreatedTab;
 	return CreatedTab.ToSharedRef();
 }
diff --git a/Engine/Source/Editor/EditorFramework/Private/Toolkits/AssetEditorModeUILayer.cpp b/Engine/Source/Editor/EditorFramework/Private/Toolkits/AssetEditorModeUILayer.cpp
--- a/Engine/Source/Editor/EditorFramework/Private/Toolkits/AssetEditorModeUILayer.cpp
+++ b/Engine/Source/Editor/EditorFramework/Private/Toolkits/AssetEditorModeUILayer.cpp
@@ -80,7 +80,7 @@ TSharedPtr<FTabManager> FAssetEditorModeUILayer::GetTabManager()
 }
 
 
-TSharedPtr<FWorkspaceItem> FAssetEditorModeUILayer::GetModeMenuCategory()
+TSharedPtr<FWorkspaceItem> FAssetEditorModeUILayer::GetModeMenuCategory() const
 {
 	return TSharedPtr<FWorkspaceItem>();
 }
@@ -104,6 +104,7 @@ void FAssetEditorModeUILayer::RegisterModeTabSpawner(const FName TabID)
 		FCanSpawnTab::CreateSP(this, &FAssetEditorModeUILayer::CanSpawnStoredTab, TabID))
 		.SetDisplayNameAttribute(MakeAttributeSP(this, &FAssetEditorModeUILayer::GetTabSpawnerName, TabID))
 		.SetTooltipTextAttribute(MakeAttributeSP(this, &FAssetEditorModeUILayer::GetTabSpawnerTooltip, TabID))
+		.SetIcon(GetTabSpawnerIcon(TabID))
 		.SetAutoGenerateMenuEntry(bShowMenuOption)
 		.SetGroup(MenuGroup);
 }
@@ -163,5 +164,15 @@ FText FAssetEditorModeUILayer::GetTabSpawnerTooltip(const FName TabID) const
 	return FText::GetEmpty();
 }
 
+const FSlateIcon& FAssetEditorModeUILayer::GetTabSpawnerIcon(const FName TabID) const
+{
+	if (RequestedTabInfo[TabID].TabIcon.IsSet())
+	{
+		return RequestedTabInfo[TabID].TabIcon;
+	}
+
+	return GetModeMenuCategory()->GetIcon();
+}
+
 #undef LOCTEXT_NAMESPACE
 
diff --git a/Engine/Source/Editor/EditorFramework/Public/Toolkits/AssetEditorModeUILayer.h b/Engine/Source/Editor/EditorFramework/Public/Toolkits/AssetEditorModeUILayer.h
--- a/Engine/Source/Editor/EditorFramework/Public/Toolkits/AssetEditorModeUILayer.h
+++ b/Engine/Source/Editor/EditorFramework/Public/Toolkits/AssetEditorModeUILayer.h
@@ -34,7 +34,7 @@ public:
 	/** Called by SLevelEditor to notify the toolbox about an existing toolkit no longer being hosted */
 	virtual void OnToolkitHostingFinished(const TSharedRef<IToolkit>& Toolkit);
 	virtual TSharedPtr<FTabManager> GetTabManager();
-	virtual TSharedPtr<FWorkspaceItem> GetModeMenuCategory();
+	virtual TSharedPtr<FWorkspaceItem> GetModeMenuCategory() const;
 	virtual void SetModePanelInfo(const FName InTabSpawnerID, const FMinorTabConfig& InTabInfo);
 	virtual TMap<FName, TWeakPtr<SDockTab>> GetSpawnedTabs();
 	virtual FSimpleDelegate& ToolkitHostReadyForUI()
@@ -60,6 +60,7 @@ protected:
 	bool CanSpawnStoredTab(const FSpawnTabArgs& Args, const FName TabID);
 	FText GetTabSpawnerName(const FName TabID) const;
 	FText GetTabSpawnerTooltip(const FName TabID) const;
+	const FSlateIcon& GetTabSpawnerIcon(const FName TabID) const;
 	virtual void RegisterLayoutExtensions(FLayoutExtender& Extender) {};
 
 protected:
diff --git a/Engine/Source/Editor/LevelEditor/Private/SLevelEditorToolBox.cpp b/Engine/Source/Editor/LevelEditor/Private/SLevelEditorToolBox.cpp
--- a/Engine/Source/Editor/LevelEditor/Private/SLevelEditorToolBox.cpp
+++ b/Engine/Source/Editor/LevelEditor/Private/SLevelEditorToolBox.cpp
@@ -78,7 +78,7 @@ void FLevelEditorModeUILayer::RegisterLayoutExtensions(FLayoutExtender& Extender
 
 
 
-TSharedPtr<FWorkspaceItem> FLevelEditorModeUILayer::GetModeMenuCategory()
+TSharedPtr<FWorkspaceItem> FLevelEditorModeUILayer::GetModeMenuCategory() const
 {
 	const IWorkspaceMenuStructure& MenuStructure = WorkspaceMenu::GetMenuStructure();
 	return MenuStructure.GetLevelEditorModesCategory();
diff --git a/Engine/Source/Editor/LevelEditor/Private/SLevelEditorToolBox.h b/Engine/Source/Editor/LevelEditor/Private/SLevelEditorToolBox.h
--- a/Engine/Source/Editor/LevelEditor/Private/SLevelEditorToolBox.h
+++ b/Engine/Source/Editor/LevelEditor/Private/SLevelEditorToolBox.h
@@ -31,7 +31,7 @@ public:
 	virtual void OnToolkitHostingStarted(const TSharedRef<IToolkit>& Toolkit) override;
 	virtual void OnToolkitHostingFinished(const TSharedRef<IToolkit>& Toolkit) override;
 
-	virtual TSharedPtr<FWorkspaceItem> GetModeMenuCategory() override;
+	virtual TSharedPtr<FWorkspaceItem> GetModeMenuCategory() const override;
 	virtual const FName GetStatusBarName() const override
 	{
 		static const FName LevelEditorStatusBarName = "LevelEditor.StatusBar";
