fix undo/redo for geometry collection "New Asset" and "Reset" actions

#rnx

diff --git a/Engine/Plugins/Experimental/ChaosEditor/Source/FractureEditor/Private/FractureToolGenerators.cpp b/Engine/Plugins/Experimental/ChaosEditor/Source/FractureEditor/Private/FractureToolGenerators.cpp
--- a/Engine/Plugins/Experimental/ChaosEditor/Source/FractureEditor/Private/FractureToolGenerators.cpp
+++ b/Engine/Plugins/Experimental/ChaosEditor/Source/FractureEditor/Private/FractureToolGenerators.cpp
@@ -60,11 +60,11 @@ void UFractureToolGenerateAsset::Execute(TWeakPtr<FFractureEditorModeToolkit> In
 	TArray<AActor*> SelectedActors;
 	SelectedActors.Reserve(SelectionSet->Num());
 
-	FScopedTransaction Transaction(LOCTEXT("GenerateAsset", "Generate Geometry Collection Asset"));
-
 	SelectionSet->GetSelectedObjects(SelectedActors);
 
 	OpenGenerateAssetDialog(SelectedActors);
+
+	// Note: Transaction for undo history is created after the user completes the UI dialog; see OnGenerateAssetPathChosen()
 }
 
 void UFractureToolGenerateAsset::OpenGenerateAssetDialog(TArray<AActor*>& Actors)
@@ -101,8 +101,9 @@ void UFractureToolGenerateAsset::OpenGenerateAssetDialog(TArray<AActor*>& Actors
 }
 
 void UFractureToolGenerateAsset::OnGenerateAssetPathChosen(const FString& InAssetPath, TArray<AActor*> Actors)
-{
-		
+{	
+	FScopedTransaction Transaction(LOCTEXT("GenerateAsset", "Generate Geometry Collection Asset"));
+
 	//Record the path
 	int32 LastSlash = INDEX_NONE;
 	if (InAssetPath.FindLastChar('/', LastSlash))
@@ -154,6 +155,7 @@ void UFractureToolGenerateAsset::OnGenerateAssetPathChosen(const FString& InAsse
 
 		for (AActor* Actor : Actors)
 		{
+			Actor->Modify();
 			Actor->Destroy();
 		}
 	}
@@ -377,10 +379,13 @@ void UFractureToolResetAsset::Execute(TWeakPtr<FFractureEditorModeToolkit> InToo
 		return;
 	}
 
+	FScopedTransaction Transaction(LOCTEXT("ResetCollection", "Reset Geometry Collection"));
+
 	TSet<UGeometryCollectionComponent*> GeomCompSelection;
 	GetSelectedGeometryCollectionComponents(GeomCompSelection);
 	for (UGeometryCollectionComponent* GeometryCollectionComponent : GeomCompSelection)
 	{
+		GeometryCollectionComponent->Modify();
 		FGeometryCollectionEdit GeometryCollectionEdit = GeometryCollectionComponent->EditRestCollection();
 		if (UGeometryCollection* GeometryCollectionObject = GeometryCollectionEdit.GetRestCollection())
 		{
