generalize fracture tools support for updating visualizations properly on explodeamount changes

#rnx

diff --git a/Engine/Plugins/Experimental/ChaosEditor/Source/FractureEditor/Private/FractureToolBrick.cpp b/Engine/Plugins/Experimental/ChaosEditor/Source/FractureEditor/Private/FractureToolBrick.cpp
--- a/Engine/Plugins/Experimental/ChaosEditor/Source/FractureEditor/Private/FractureToolBrick.cpp
+++ b/Engine/Plugins/Experimental/ChaosEditor/Source/FractureEditor/Private/FractureToolBrick.cpp
@@ -300,43 +300,24 @@ void UFractureToolBrick::GenerateBrickTransforms(const FBox& Bounds)
 
 void UFractureToolBrick::UpdateBrickTransforms()
 {
-	FBox Bounds(ForceInitToZero);
+	TArray<FFractureToolContext> FractureContexts = GetFractureToolContexts();
 
-	USelection* SelectionSet = GEditor->GetSelectedActors();
+	ClearVisualizations();
 
-	TArray<AActor*> SelectedActors;
-	SelectedActors.Reserve(SelectionSet->Num());
-
-
-	FBox SelectedMeshBounds(ForceInit);
-
-	SelectionSet->GetSelectedObjects(SelectedActors);
-	BrickTransforms.Empty();
-	Edges.Empty();
-
-	for (AActor* Actor : SelectedActors)
+	for (const FFractureToolContext& FractureContext : FractureContexts)
 	{
-		TInlineComponentArray<UPrimitiveComponent*> PrimitiveComponents;
-		Actor->GetComponents(PrimitiveComponents);
-		for (UPrimitiveComponent* PrimitiveComponent : PrimitiveComponents)
+		if (!FractureContext.GetBounds().IsValid) // skip contexts w/ invalid bounds
 		{
-			FVector Origin;
-			FVector BoxExtent;
-			Actor->GetActorBounds(false, Origin, BoxExtent);
-
-			if (CutterSettings->bGroupFracture)
-			{
-				Bounds += FBox::BuildAABB(Origin, BoxExtent);
-			}
-			else
-			{
-				GenerateBrickTransforms(FBox::BuildAABB(Origin, BoxExtent));
-			}
+			continue;
 		}
-	}
+		int32 CollectionIdx = VisualizedCollections.Emplace(FractureContext.GetGeometryCollectionComponent());
+		int32 BoneIdx = FractureContext.GetSelection().Num() == 1 ? FractureContext.GetSelection()[0] : INDEX_NONE;
+		BricksMappings.AddMapping(CollectionIdx, BoneIdx, BrickTransforms.Num());
+		EdgesMappings.AddMapping(CollectionIdx, BoneIdx, Edges.Num());
 
-	if (CutterSettings->bGroupFracture)
-	{
+		FGeometryCollection& Collection = *FractureContext.GetGeometryCollection();
+
+		const FBox Bounds = FractureContext.GetWorldBounds();
 		GenerateBrickTransforms(Bounds);
 	}
 }
@@ -354,20 +335,19 @@ void UFractureToolBrick::FractureContextChanged()
 
 void UFractureToolBrick::Render(const FSceneView* View, FViewport* Viewport, FPrimitiveDrawInterface* PDI)
 {
-	//TODO :Plane cutting drawing
-
-	for (const FTransform& Transform : BrickTransforms)
+	EnumerateVisualizationMapping(BricksMappings, BrickTransforms.Num(), [&](int32 Idx, FVector ExplodedVector)
 	{
-		PDI->DrawPoint(Transform.GetLocation(), FLinearColor::Green, 4.f, SDPG_Foreground);
-	}
+		const FTransform& Transform = BrickTransforms[Idx];
+		PDI->DrawPoint(Transform.GetLocation() + ExplodedVector, FLinearColor::Green, 4.f, SDPG_Foreground);
+	});
 
 	if (CutterSettings->bDrawDiagram)
 	{
 		PDI->AddReserveLines(SDPG_Foreground, Edges.Num(), false, false);
-		for (int32 ii = 0, ni = Edges.Num(); ii < ni; ++ii)
+		EnumerateVisualizationMapping(EdgesMappings, Edges.Num(), [&](int32 Idx, FVector ExplodedVector)
 		{
-			PDI->DrawLine(Edges[ii].Get<0>(), Edges[ii].Get<1>(), FLinearColor(255, 0, 0), SDPG_Foreground);
-		}
+			PDI->DrawLine(Edges[Idx].Get<0>() + ExplodedVector, Edges[Idx].Get<1>() + ExplodedVector, FLinearColor(255, 0, 0), SDPG_Foreground);
+		});
 	}
 
 }
diff --git a/Engine/Plugins/Experimental/ChaosEditor/Source/FractureEditor/Private/FractureToolBrick.h b/Engine/Plugins/Experimental/ChaosEditor/Source/FractureEditor/Private/FractureToolBrick.h
--- a/Engine/Plugins/Experimental/ChaosEditor/Source/FractureEditor/Private/FractureToolBrick.h
+++ b/Engine/Plugins/Experimental/ChaosEditor/Source/FractureEditor/Private/FractureToolBrick.h
@@ -99,6 +99,16 @@ public:
 
 	void GenerateBrickTransforms(const FBox& Bounds);
 
+protected:
+	virtual void ClearVisualizations() override
+	{
+		Super::ClearVisualizations();
+		BrickTransforms.Empty();
+		Edges.Empty();
+		Boxes.Empty();
+		BricksMappings.Empty();
+		EdgesMappings.Empty();
+	}
 
 private:
 	UPROPERTY()
@@ -111,4 +121,5 @@ private:
 	TArray<FTransform> BrickTransforms;
 	TArray<TTuple<FVector, FVector>> Edges;
 	TArray<TTuple<FVector, FVector>> Boxes;
+	FVisualizationMappings BricksMappings, EdgesMappings;
 };
\ No newline at end of file
diff --git a/Engine/Plugins/Experimental/ChaosEditor/Source/FractureEditor/Private/FractureToolConvex.cpp b/Engine/Plugins/Experimental/ChaosEditor/Source/FractureEditor/Private/FractureToolConvex.cpp
--- a/Engine/Plugins/Experimental/ChaosEditor/Source/FractureEditor/Private/FractureToolConvex.cpp
+++ b/Engine/Plugins/Experimental/ChaosEditor/Source/FractureEditor/Private/FractureToolConvex.cpp
@@ -61,8 +61,7 @@ void UFractureToolConvex::FractureContextChanged()
 {
 	TArray<FFractureToolContext> FractureContexts = GetFractureToolContexts();
 
-	HullPoints.Reset();
-	HullEdges.Reset();
+	ClearVisualizations();
 
 	for (const FFractureToolContext& FractureContext : FractureContexts)
 	{
@@ -74,21 +73,19 @@ void UFractureToolConvex::FractureContextChanged()
 			continue;
 		}
 
+		int32 CollectionIdx = VisualizedCollections.Add(FractureContext.GetGeometryCollectionComponent());
+
 		FTransform OuterTransform = FractureContext.GetTransform();
 		for (int32 TransformIdx : FractureContext.GetSelection())
 		{
 			FTransform InnerTransform = GeometryCollectionAlgo::GlobalMatrix(Collection.Transform, Collection.Parent, TransformIdx);
-			if (Collection.HasAttribute("ExplodedVector", FGeometryCollection::TransformGroup))
-			{
-				TManagedArray<FVector3f>& ExplodedVectors = Collection.GetAttribute<FVector3f>("ExplodedVector", FGeometryCollection::TransformGroup);
-				InnerTransform = InnerTransform * FTransform(ExplodedVectors[TransformIdx]);
-			}
-
 			FTransform CombinedTransform = InnerTransform * OuterTransform;
 
 			TManagedArray<TSet<int32>>& TransformToConvexIndices = Collection.GetAttribute<TSet<int32>>("TransformToConvexIndices", FTransformCollection::TransformGroup);
 			TManagedArray<TUniquePtr<Chaos::FConvex>>& ConvexHull = Collection.GetAttribute<TUniquePtr<Chaos::FConvex>>("ConvexHull", "Convex");
 
+			EdgesMappings.AddMapping(CollectionIdx, TransformIdx, HullEdges.Num());
+
 			for (int32 ConvexIdx : TransformToConvexIndices[TransformIdx])
 			{
 				int32 HullPtsStart = HullPoints.Num();
@@ -115,12 +112,13 @@ void UFractureToolConvex::FractureContextChanged()
 
 void UFractureToolConvex::Render(const FSceneView* View, FViewport* Viewport, FPrimitiveDrawInterface* PDI)
 {
-	for (TPair<int32,int32> Edge : HullEdges)
+	EnumerateVisualizationMapping(EdgesMappings, HullEdges.Num(), [&](int32 Idx, FVector ExplodedVector)
 	{
-		FVector P1 = HullPoints[Edge.Key];
-		FVector P2 = HullPoints[Edge.Value];
+		TPair<int32, int32> Edge = HullEdges[Idx];
+		FVector P1 = HullPoints[Edge.Key] + ExplodedVector;
+		FVector P2 = HullPoints[Edge.Value] + ExplodedVector;
 		PDI->DrawLine(P1, P2, FLinearColor::Green, SDPG_Foreground, 0.0f, 0.001f);
-	}
+	});
 }
 
 void UFractureToolConvex::PostEditChangeChainProperty(struct FPropertyChangedChainEvent& PropertyChangedEvent)
diff --git a/Engine/Plugins/Experimental/ChaosEditor/Source/FractureEditor/Private/FractureToolConvex.h b/Engine/Plugins/Experimental/ChaosEditor/Source/FractureEditor/Private/FractureToolConvex.h
--- a/Engine/Plugins/Experimental/ChaosEditor/Source/FractureEditor/Private/FractureToolConvex.h
+++ b/Engine/Plugins/Experimental/ChaosEditor/Source/FractureEditor/Private/FractureToolConvex.h
@@ -63,6 +63,16 @@ protected:
 	UPROPERTY(EditAnywhere, Category = Slicing)
 	TObjectPtr<UFractureConvexSettings> ConvexSettings;
 
+	virtual void ClearVisualizations() override
+	{
+		Super::ClearVisualizations();
+		HullPoints.Empty();
+		HullEdges.Empty();
+		EdgesMappings.Empty();
+	}
+
 	TArray<FVector> HullPoints;
 	TArray<TPair<int32, int32>> HullEdges;
+	FVisualizationMappings EdgesMappings;
+
 };
diff --git a/Engine/Plugins/Experimental/ChaosEditor/Source/FractureEditor/Private/FractureToolCutter.cpp b/Engine/Plugins/Experimental/ChaosEditor/Source/FractureEditor/Private/FractureToolCutter.cpp
--- a/Engine/Plugins/Experimental/ChaosEditor/Source/FractureEditor/Private/FractureToolCutter.cpp
+++ b/Engine/Plugins/Experimental/ChaosEditor/Source/FractureEditor/Private/FractureToolCutter.cpp
@@ -220,19 +220,20 @@ void UFractureToolVoronoiCutterBase::Render(const FSceneView* View, FViewport* V
 {
 	if (CutterSettings->bDrawSites)
 	{
-		for (const FVector& Site : VoronoiSites)
+		EnumerateVisualizationMapping(SitesMappings, VoronoiSites.Num(), [&](int32 Idx, FVector ExplodedVector)
 		{
-			PDI->DrawPoint(Site, FLinearColor::Green, 4.f, SDPG_Foreground);
-		}
+			PDI->DrawPoint(VoronoiSites[Idx] + ExplodedVector, FLinearColor::Green, 4.f, SDPG_Foreground);
+		});
 	}
 
 	if (CutterSettings->bDrawDiagram)
 	{
 		PDI->AddReserveLines(SDPG_Foreground, VoronoiEdges.Num(), false, false);
-		for (int32 ii = 0, ni = VoronoiEdges.Num(); ii < ni; ++ii)
+		EnumerateVisualizationMapping(EdgesMappings, VoronoiEdges.Num(), [&](int32 Idx, FVector ExplodedVector)
 		{
-			PDI->DrawLine(VoronoiEdges[ii].Get<0>(), VoronoiEdges[ii].Get<1>(), Colors[CellMember[ii] % 100], SDPG_Foreground);
-		}
+			PDI->DrawLine(VoronoiEdges[Idx].Get<0>() + ExplodedVector,
+				VoronoiEdges[Idx].Get<1>() + ExplodedVector, Colors[CellMember[Idx] % 100], SDPG_Foreground);
+		});
 	}
 }
 
@@ -241,9 +242,7 @@ void UFractureToolVoronoiCutterBase::FractureContextChanged()
 	UpdateDefaultRandomSeed();
 	TArray<FFractureToolContext> FractureContexts = GetFractureToolContexts();
 
-	VoronoiSites.Empty();
-	CellMember.Empty();
-	VoronoiEdges.Empty();
+	ClearVisualizations();
 
 	int32 MaxSitesToShowEdges = 100000; // computing all the voronoi diagrams can make the program non-responsive above this
 	int32 MaxSitesToShowSites = 1000000; // PDI struggles to render the site positions above this
@@ -255,6 +254,11 @@ void UFractureToolVoronoiCutterBase::FractureContextChanged()
 		{
 			continue;
 		}
+		int32 CollectionIdx = VisualizedCollections.Emplace(FractureContext.GetGeometryCollectionComponent());
+		int32 BoneIdx = FractureContext.GetSelection().Num() == 1 ? FractureContext.GetSelection()[0] : INDEX_NONE;
+		SitesMappings.AddMapping(CollectionIdx, BoneIdx, VoronoiSites.Num());
+		EdgesMappings.AddMapping(CollectionIdx, BoneIdx, VoronoiEdges.Num());
+
 		// Generate voronoi diagrams and cache visualization info
 		TArray<FVector> LocalVoronoiSites;
 		GenerateVoronoiSites(FractureContext, LocalVoronoiSites);
@@ -262,9 +266,7 @@ void UFractureToolVoronoiCutterBase::FractureContextChanged()
 		if (LocalVoronoiSites.Num() * FractureContexts.Num() > MaxSitesToShowSites || VoronoiSites.Num() + LocalVoronoiSites.Num() > MaxSitesToShowSites)
 		{
 			UE_LOG(LogFractureTool, Warning, TEXT("Voronoi diagram(s) number of sites too large; will not display Voronoi diagram sites"));
-			VoronoiSites.Empty();
-			CellMember.Empty();
-			VoronoiEdges.Empty();
+			ClearVisualizations();
 			break;
 		}
 		VoronoiSites.Append(LocalVoronoiSites);
@@ -273,6 +275,7 @@ void UFractureToolVoronoiCutterBase::FractureContextChanged()
 			UE_LOG(LogFractureTool, Warning, TEXT("Voronoi diagram(s) number of sites too large; will not display Voronoi diagram edges"));
 			VoronoiEdges.Empty();
 			CellMember.Empty();
+			EdgesMappings.Empty();
 			bEstAboveMaxSites = true;
 		}
 		else
diff --git a/Engine/Plugins/Experimental/ChaosEditor/Source/FractureEditor/Private/FractureToolCutter.h b/Engine/Plugins/Experimental/ChaosEditor/Source/FractureEditor/Private/FractureToolCutter.h
--- a/Engine/Plugins/Experimental/ChaosEditor/Source/FractureEditor/Private/FractureToolCutter.h
+++ b/Engine/Plugins/Experimental/ChaosEditor/Source/FractureEditor/Private/FractureToolCutter.h
@@ -158,11 +158,23 @@ protected:
 
 	FBox GetVoronoiBounds(const FFractureToolContext& FractureContext, const TArray<FVector>& Sites) const;
 
+	virtual void ClearVisualizations()
+	{
+		Super::ClearVisualizations();
+		CellMember.Empty();
+		VoronoiEdges.Empty();
+		VoronoiSites.Empty();
+		SitesMappings.Empty();
+		EdgesMappings.Empty();
+	}
+
 private:
 	TArray<int32> CellMember;
 	TArray<TTuple<FVector, FVector>> VoronoiEdges;
 	TArray<FVector> VoronoiSites;
 	TArray<FLinearColor> Colors;
+
+	FVisualizationMappings SitesMappings, EdgesMappings;
 };
 
 
diff --git a/Engine/Plugins/Experimental/ChaosEditor/Source/FractureEditor/Private/FractureToolFixTinyGeo.cpp b/Engine/Plugins/Experimental/ChaosEditor/Source/FractureEditor/Private/FractureToolFixTinyGeo.cpp
--- a/Engine/Plugins/Experimental/ChaosEditor/Source/FractureEditor/Private/FractureToolFixTinyGeo.cpp
+++ b/Engine/Plugins/Experimental/ChaosEditor/Source/FractureEditor/Private/FractureToolFixTinyGeo.cpp
@@ -40,45 +40,13 @@ void UFractureToolFixTinyGeo::RegisterUICommand( FFractureEditorCommands* Bindin
 	BindingContext->FixTinyGeo = UICommandInfo;
 }
 
-const TManagedArray<FVector3f>* UFractureToolFixTinyGeo::GetExplodedVectors(const UGeometryCollection* GeometryCollectionObject)
-{
-	if (GeometryCollectionObject)
-	{
-		const FGeometryCollection& Collection = *GeometryCollectionObject->GetGeometryCollection();
-		if (Collection.HasAttribute("ExplodedVector", FGeometryCollection::TransformGroup))
-		{
-			return &Collection.GetAttribute<FVector3f>("ExplodedVector", FGeometryCollection::TransformGroup);
-		}
-	}
-	return nullptr;
-}
-
 void UFractureToolFixTinyGeo::Render(const FSceneView* View, FViewport* Viewport, FPrimitiveDrawInterface* PDI)
 {
-	if (!ensure(ToRemoveCollectionBones.Num() == ToRemoveBounds.Num()))
+	EnumerateVisualizationMapping(ToRemoveMappings, ToRemoveBounds.Num(), [&](int32 Idx, FVector ExplodedVector)
 	{
-		return;
-	}
-	const TManagedArray<FVector3f>* ExplodedVectors = nullptr;
-
-	int32 LastCollectionIdx = INDEX_NONE;
-	for (int32 Idx = 0; Idx < ToRemoveBounds.Num(); Idx++)
-	{
-		TPair<int32, int32> CollectBones = ToRemoveCollectionBones[Idx];
-		if (CollectBones.Key != LastCollectionIdx)
-		{
-			ExplodedVectors = GetExplodedVectors(VisCollections[CollectBones.Key]);
-			LastCollectionIdx = CollectBones.Key;
-		}
-		int32 BoneIdx = CollectBones.Value;
 		const FBox& Box = ToRemoveBounds[Idx];
-		FVector Offset = FVector::ZeroVector;
-		if (ExplodedVectors && BoneIdx < ExplodedVectors->Num())
-		{
-			Offset = (*ExplodedVectors)[BoneIdx];
-		}
-		FVector B000 = Box.Min + Offset;
-		FVector B111 = Box.Max + Offset;
+		FVector B000 = Box.Min + ExplodedVector;
+		FVector B111 = Box.Max + ExplodedVector;
 		FVector B011(B000.X, B111.Y, B111.Z);
 		FVector B101(B111.X, B000.Y, B111.Z);
 		FVector B110(B111.X, B111.Y, B000.Z);
@@ -97,7 +65,7 @@ void UFractureToolFixTinyGeo::Render(const FSceneView* View, FViewport* Viewport
 		PDI->DrawLine(B110, B010, FLinearColor::Red, SDPG_Foreground, 0.0f, 0.001f);
 		PDI->DrawLine(B100, B101, FLinearColor::Red, SDPG_Foreground, 0.0f, 0.001f);
 		PDI->DrawLine(B010, B011, FLinearColor::Red, SDPG_Foreground, 0.0f, 0.001f);
-	}
+	});
 	
 }
 
@@ -115,14 +83,12 @@ void UFractureToolFixTinyGeo::FractureContextChanged()
 	UpdateDefaultRandomSeed();
 	TArray<FFractureToolContext> FractureContexts = GetFractureToolContexts();
 
-	ToRemoveCollectionBones.Reset();
-	ToRemoveBounds.Reset();
-	VisCollections.Reset();
+	ClearVisualizations();
 
 	for (const FFractureToolContext& FractureContext : FractureContexts)
 	{
 		FGeometryCollection& Collection = *FractureContext.GetGeometryCollection();
-		int CollectionIdx = VisCollections.Emplace(FractureContext.GetGeometryCollectionComponent()->GetRestCollection());
+		int CollectionIdx = VisualizedCollections.Emplace(FractureContext.GetGeometryCollectionComponent());
 
 		TArray<double> Volumes;
 		FindBoneVolumes(
@@ -164,7 +130,7 @@ void UFractureToolFixTinyGeo::FractureContextChanged()
 			{
 				Bounds += CombinedTransform.TransformPosition(Collection.Vertex[VIdx]);
 			}
-			ToRemoveCollectionBones.Add(TPair<int32, int32>(CollectionIdx, TransformIdx));
+			ToRemoveMappings.AddMapping(CollectionIdx, TransformIdx, ToRemoveBounds.Num());
 			ToRemoveBounds.Add(Bounds);
 		}
 	}
diff --git a/Engine/Plugins/Experimental/ChaosEditor/Source/FractureEditor/Private/FractureToolFixTinyGeo.h b/Engine/Plugins/Experimental/ChaosEditor/Source/FractureEditor/Private/FractureToolFixTinyGeo.h
--- a/Engine/Plugins/Experimental/ChaosEditor/Source/FractureEditor/Private/FractureToolFixTinyGeo.h
+++ b/Engine/Plugins/Experimental/ChaosEditor/Source/FractureEditor/Private/FractureToolFixTinyGeo.h
@@ -72,16 +72,21 @@ public:
 	virtual void FractureContextChanged() override;
 	virtual int32 ExecuteFracture(const FFractureToolContext& FractureContext) override;
 
+protected:
+	virtual void ClearVisualizations() override
+	{
+		Super::ClearVisualizations();
+		ToRemoveBounds.Empty();
+		ToRemoveMappings.Empty();
+	}
+
 private:
 
 	UPROPERTY(EditAnywhere, Category = FixGeo)
 	TObjectPtr<UFractureTinyGeoSettings> TinyGeoSettings;
 
-	TArray<TPair<int32, int32>> ToRemoveCollectionBones; // Pairs of [index into VisCollections, bone index]
 	TArray<FBox> ToRemoveBounds; // Bounds in global space but without exploded vectors applied
-	TArray<TObjectPtr<const UGeometryCollection>> VisCollections; // source geometry collections
-
-	const TManagedArray<FVector3f>* GetExplodedVectors(const UGeometryCollection* GeometryCollectionObject);
+	FVisualizationMappings ToRemoveMappings;
 
 	double GetMinVolume(TArray<double>& Volumes);
 	const double VolDimScale = .01; // compute volumes in meters instead of cm, for saner units at typical scales
diff --git a/Engine/Plugins/Experimental/ChaosEditor/Source/FractureEditor/Private/FractureToolPlaneCut.cpp b/Engine/Plugins/Experimental/ChaosEditor/Source/FractureEditor/Private/FractureToolPlaneCut.cpp
--- a/Engine/Plugins/Experimental/ChaosEditor/Source/FractureEditor/Private/FractureToolPlaneCut.cpp
+++ b/Engine/Plugins/Experimental/ChaosEditor/Source/FractureEditor/Private/FractureToolPlaneCut.cpp
@@ -61,9 +61,9 @@ void UFractureToolPlaneCut::Render(const FSceneView* View, FViewport* Viewport,
 	if (CutterSettings->bDrawDiagram)
 	{
 		// Draw a point centered at plane origin, and a square on the plane around it.
-		auto DrawPlane = [&PDI](const FTransform& Transform, float PlaneSize)
+		auto DrawPlane = [&PDI](const FTransform& Transform, float PlaneSize, FVector Offset)
 		{
-			FVector Center = Transform.GetLocation();
+			FVector Center = Transform.GetLocation() + Offset;
 			FVector X = (PlaneSize * .5) * Transform.GetUnitAxis(EAxis::X);
 			FVector Y = (PlaneSize * .5) * Transform.GetUnitAxis(EAxis::Y);
 			FVector Corners[4]
@@ -83,14 +83,18 @@ void UFractureToolPlaneCut::Render(const FSceneView* View, FViewport* Viewport,
 		if (GizmoSettings->IsGizmoEnabled())
 		{
 			const FTransform& Transform = GizmoSettings->GetTransform();
-			DrawPlane(Transform, 100.f);
+			EnumerateVisualizationMapping(PlanesMappings, RenderCuttingPlanesTransforms.Num(), [&](int32 Idx, FVector ExplodedVector)
+			{
+				DrawPlane(Transform, 100.f, ExplodedVector);
+			});
 		}
 		else // draw from computed transforms
 		{
-			for (const FTransform& Transform : RenderCuttingPlanesTransforms)
+			EnumerateVisualizationMapping(PlanesMappings, RenderCuttingPlanesTransforms.Num(), [&](int32 Idx, FVector ExplodedVector)
 			{
-				DrawPlane(Transform, RenderCuttingPlaneSize);
-			}
+				const FTransform& Transform = RenderCuttingPlanesTransforms[Idx];
+				DrawPlane(Transform, RenderCuttingPlaneSize, ExplodedVector);
+			});
 		}
 	}
 }
@@ -110,17 +114,20 @@ void UFractureToolPlaneCut::FractureContextChanged()
 	UpdateDefaultRandomSeed();
 	TArray<FFractureToolContext> FractureContexts = GetFractureToolContexts();
 
-	RenderCuttingPlanesTransforms.Empty();
+	ClearVisualizations();
 
 	RenderCuttingPlaneSize = FLT_MAX;
 	for (FFractureToolContext& FractureContext : FractureContexts)
 	{
-		// Move the local bounds to the actor so we we'll draw in the correct location
 		FBox Bounds = FractureContext.GetWorldBounds();
 		if (!Bounds.IsValid)
 		{
 			continue;
 		}
+		int32 CollectionIdx = VisualizedCollections.Add(FractureContext.GetGeometryCollectionComponent());
+		int32 BoneIdx = FractureContext.GetSelection().Num() == 1 ? FractureContext.GetSelection()[0] : INDEX_NONE;
+		PlanesMappings.AddMapping(CollectionIdx, BoneIdx, RenderCuttingPlanesTransforms.Num());
+
 		GenerateSliceTransforms(FractureContext, RenderCuttingPlanesTransforms);
 
 		if (Bounds.GetExtent().GetMax() < RenderCuttingPlaneSize)
diff --git a/Engine/Plugins/Experimental/ChaosEditor/Source/FractureEditor/Private/FractureToolPlaneCut.h b/Engine/Plugins/Experimental/ChaosEditor/Source/FractureEditor/Private/FractureToolPlaneCut.h
--- a/Engine/Plugins/Experimental/ChaosEditor/Source/FractureEditor/Private/FractureToolPlaneCut.h
+++ b/Engine/Plugins/Experimental/ChaosEditor/Source/FractureEditor/Private/FractureToolPlaneCut.h
@@ -51,6 +51,14 @@ public:
 	virtual void Setup() override;
 	virtual void Shutdown() override;
 
+protected:
+	virtual void ClearVisualizations() override
+	{
+		Super::ClearVisualizations();
+		RenderCuttingPlanesTransforms.Empty();
+		PlanesMappings.Empty();
+	}
+
 private:
 	// Slicing
 	UPROPERTY(EditAnywhere, Category = Slicing)
@@ -63,6 +71,7 @@ private:
 
 	float RenderCuttingPlaneSize;
 	TArray<FTransform> RenderCuttingPlanesTransforms;
+	FVisualizationMappings PlanesMappings;
 };
 
 
diff --git a/Engine/Plugins/Experimental/ChaosEditor/Source/FractureEditor/Private/FractureToolRecomputeNormals.cpp b/Engine/Plugins/Experimental/ChaosEditor/Source/FractureEditor/Private/FractureToolRecomputeNormals.cpp
--- a/Engine/Plugins/Experimental/ChaosEditor/Source/FractureEditor/Private/FractureToolRecomputeNormals.cpp
+++ b/Engine/Plugins/Experimental/ChaosEditor/Source/FractureEditor/Private/FractureToolRecomputeNormals.cpp
@@ -44,9 +44,9 @@ void UFractureToolRecomputeNormals::Render(const FSceneView* View, FViewport* Vi
 {
 	if (NormalsSettings->bShowTangents || NormalsSettings->bShowNormals)
 	{
-		for (int32 Idx = 0; Idx < DisplayVertices.Num(); Idx++)
+		EnumerateVisualizationMapping(PointsMappings, DisplayVertices.Num(), [&](int32 Idx, FVector ExplodedVector)
 		{
-			FVector Point = DisplayVertices[Idx];
+			FVector Point = DisplayVertices[Idx] + ExplodedVector;
 			FVector Normal = DisplayNormals[Idx] * NormalsSettings->Length;
 			FVector TanU = DisplayTanUs[Idx] * NormalsSettings->Length;
 			FVector TanV = DisplayTanVs[Idx] * NormalsSettings->Length;
@@ -59,7 +59,7 @@ void UFractureToolRecomputeNormals::Render(const FSceneView* View, FViewport* Vi
 				PDI->DrawLine(Point, Point + TanU, FLinearColor::Green, SDPG_Foreground);
 				PDI->DrawLine(Point, Point + TanV, FLinearColor::Blue, SDPG_Foreground);
 			}
-		}
+		});
 	}
 }
 
@@ -75,10 +75,8 @@ void UFractureToolRecomputeNormals::FractureContextChanged()
 {
 	TArray<FFractureToolContext> FractureContexts = GetFractureToolContexts();
 
-	DisplayVertices.Reset();
-	DisplayNormals.Reset();
-	DisplayTanUs.Reset();
-	DisplayTanVs.Reset();
+	ClearVisualizations();
+
 	for (const FFractureToolContext& FractureContext : FractureContexts)
 	{
 		FGeometryCollection& Collection = *FractureContext.GetGeometryCollection();
@@ -87,13 +85,10 @@ void UFractureToolRecomputeNormals::FractureContextChanged()
 		FTransform OuterTransform = FractureContext.GetTransform();
 		for (int32 TransformIdx : FractureContext.GetSelection())
 		{
+			int32 CollectionIdx = VisualizedCollections.Add(FractureContext.GetGeometryCollectionComponent());
+			PointsMappings.AddMapping(CollectionIdx, TransformIdx, DisplayVertices.Num());
+
 			FTransform InnerTransform = GeometryCollectionAlgo::GlobalMatrix(Collection.Transform, Collection.Parent, TransformIdx);
-			if (Collection.HasAttribute("ExplodedVector", FGeometryCollection::TransformGroup))
-			{
-				TManagedArray<FVector3f>& ExplodedVectors = Collection.GetAttribute<FVector3f>("ExplodedVector", FGeometryCollection::TransformGroup);
-				InnerTransform = InnerTransform * FTransform(ExplodedVectors[TransformIdx]);
-			}
-			
 			FTransform CombinedTransform = InnerTransform * OuterTransform;
 			int32 GeometryIdx = Collection.TransformToGeometryIndex[TransformIdx];
 			if (GeometryIdx > -1)
diff --git a/Engine/Plugins/Experimental/ChaosEditor/Source/FractureEditor/Private/FractureToolRecomputeNormals.h b/Engine/Plugins/Experimental/ChaosEditor/Source/FractureEditor/Private/FractureToolRecomputeNormals.h
--- a/Engine/Plugins/Experimental/ChaosEditor/Source/FractureEditor/Private/FractureToolRecomputeNormals.h
+++ b/Engine/Plugins/Experimental/ChaosEditor/Source/FractureEditor/Private/FractureToolRecomputeNormals.h
@@ -72,10 +72,21 @@ protected:
 	UPROPERTY()
 	TObjectPtr<UFractureRecomputeNormalsSettings> NormalsSettings;
 
+	virtual void ClearVisualizations() override
+	{
+		Super::ClearVisualizations();
+		DisplayVertices.Empty();
+		DisplayNormals.Empty();
+		DisplayTanUs.Empty();
+		DisplayTanVs.Empty();
+		PointsMappings.Empty();
+	}
+
 	TArray<FVector> DisplayVertices;
 	TArray<FVector> DisplayNormals;
 	TArray<FVector> DisplayTanUs;
 	TArray<FVector> DisplayTanVs;
+	FVisualizationMappings PointsMappings;
 
 };
 
diff --git a/Engine/Plugins/Experimental/ChaosEditor/Source/FractureEditor/Private/FractureToolResample.cpp b/Engine/Plugins/Experimental/ChaosEditor/Source/FractureEditor/Private/FractureToolResample.cpp
--- a/Engine/Plugins/Experimental/ChaosEditor/Source/FractureEditor/Private/FractureToolResample.cpp
+++ b/Engine/Plugins/Experimental/ChaosEditor/Source/FractureEditor/Private/FractureToolResample.cpp
@@ -40,11 +40,11 @@ void UFractureToolResample::RegisterUICommand( FFractureEditorCommands* BindingC
 
 void UFractureToolResample::Render(const FSceneView* View, FViewport* Viewport, FPrimitiveDrawInterface* PDI)
 {
-	for (FVector Point : GeneratedPoints)
+	EnumerateVisualizationMapping(PointsMappings, GeneratedPoints.Num(), [&](int32 Idx, FVector ExplodedVector)
 	{
-		PDI->DrawPoint(Point, FLinearColor::Green, 2.f, SDPG_Foreground);
-	}
-	
+		FVector Point = GeneratedPoints[Idx];
+		PDI->DrawPoint(Point + ExplodedVector, FLinearColor::Green, 2.f, SDPG_Foreground);
+	});
 }
 
 TArray<UObject*> UFractureToolResample::GetSettingsObjects() const
@@ -60,7 +60,8 @@ void UFractureToolResample::FractureContextChanged()
 	UpdateDefaultRandomSeed();
 	TArray<FFractureToolContext> FractureContexts = GetFractureToolContexts();
 
-	GeneratedPoints.Reset();
+	ClearVisualizations();
+
 	for (const FFractureToolContext& FractureContext : FractureContexts)
 	{
 		FGeometryCollection& Collection = *FractureContext.GetGeometryCollection();
@@ -69,13 +70,10 @@ void UFractureToolResample::FractureContextChanged()
 		FTransform OuterTransform = FractureContext.GetTransform();
 		for (int32 TransformIdx : FractureContext.GetSelection())
 		{
+			int32 CollectionIdx = VisualizedCollections.Add(FractureContext.GetGeometryCollectionComponent());
+			PointsMappings.AddMapping(CollectionIdx, TransformIdx, GeneratedPoints.Num());
+
 			FTransform InnerTransform = GeometryCollectionAlgo::GlobalMatrix(Collection.Transform, Collection.Parent, TransformIdx);
-			if (Collection.HasAttribute("ExplodedVector", FGeometryCollection::TransformGroup))
-			{
-				TManagedArray<FVector3f>& ExplodedVectors = Collection.GetAttribute<FVector3f>("ExplodedVector", FGeometryCollection::TransformGroup);
-				InnerTransform = InnerTransform * FTransform(ExplodedVectors[TransformIdx]);
-			}
-			
 			FTransform CombinedTransform = InnerTransform * OuterTransform;
 			int32 GeometryIdx = Collection.TransformToGeometryIndex[TransformIdx];
 			int32 VertStart = Collection.VertexStart[GeometryIdx];
diff --git a/Engine/Plugins/Experimental/ChaosEditor/Source/FractureEditor/Private/FractureToolResample.h b/Engine/Plugins/Experimental/ChaosEditor/Source/FractureEditor/Private/FractureToolResample.h
--- a/Engine/Plugins/Experimental/ChaosEditor/Source/FractureEditor/Private/FractureToolResample.h
+++ b/Engine/Plugins/Experimental/ChaosEditor/Source/FractureEditor/Private/FractureToolResample.h
@@ -34,11 +34,18 @@ public:
 	virtual void FractureContextChanged() override;
 	virtual int32 ExecuteFracture(const FFractureToolContext& FractureContext) override;
 
-private:
+protected:
+	virtual void ClearVisualizations() override
+	{
+		Super::ClearVisualizations();
+		GeneratedPoints.Empty();
+		PointsMappings.Empty();
+	}
 
+private:
 
 	TArray<FVector> GeneratedPoints;
-
+	FVisualizationMappings PointsMappings;
 };
 
 
diff --git a/Engine/Plugins/Experimental/ChaosEditor/Source/FractureEditor/Private/FractureToolSlice.cpp b/Engine/Plugins/Experimental/ChaosEditor/Source/FractureEditor/Private/FractureToolSlice.cpp
--- a/Engine/Plugins/Experimental/ChaosEditor/Source/FractureEditor/Private/FractureToolSlice.cpp
+++ b/Engine/Plugins/Experimental/ChaosEditor/Source/FractureEditor/Private/FractureToolSlice.cpp
@@ -114,31 +114,40 @@ void UFractureToolSlice::Render(const FSceneView* View, FViewport* Viewport, FPr
 {
 	if (CutterSettings->bDrawDiagram)
 	{
-	for (const FTransform& Transform : RenderCuttingPlanesTransforms)
-	{
-		PDI->DrawPoint(Transform.GetLocation(), FLinearColor::Green, 4.f, SDPG_Foreground);
+		EnumerateVisualizationMapping(PlanesMappings, RenderCuttingPlanesTransforms.Num(), [&](int32 Idx, FVector ExplodedVector)
+		{
+			const FTransform& Transform = RenderCuttingPlanesTransforms[Idx];
+			FVector Center = Transform.GetLocation() + ExplodedVector;
+			PDI->DrawPoint(Center, FLinearColor::Green, 4.f, SDPG_Foreground);
 
-		PDI->DrawLine(Transform.GetLocation(), Transform.GetLocation() + Transform.GetUnitAxis(EAxis::X) * RenderCuttingPlaneSize, FLinearColor(255, 0, 0), SDPG_Foreground);
-		PDI->DrawLine(Transform.GetLocation(), Transform.GetLocation() + Transform.GetUnitAxis(EAxis::Y) * RenderCuttingPlaneSize, FLinearColor(0, 255, 0), SDPG_Foreground);
+			PDI->DrawLine(Center, Center + Transform.GetUnitAxis(EAxis::X) * RenderCuttingPlaneSize, FLinearColor(255, 0, 0), SDPG_Foreground);
+			PDI->DrawLine(Center, Center + Transform.GetUnitAxis(EAxis::Y) * RenderCuttingPlaneSize, FLinearColor(0, 255, 0), SDPG_Foreground);
 
-		PDI->DrawLine(Transform.GetLocation() + Transform.GetUnitAxis(EAxis::X) * RenderCuttingPlaneSize, Transform.GetLocation() + Transform.GetUnitAxis(EAxis::X) * RenderCuttingPlaneSize + Transform.GetUnitAxis(EAxis::Y) * RenderCuttingPlaneSize, FLinearColor(255, 0, 0), SDPG_Foreground);
-		PDI->DrawLine(Transform.GetLocation() + Transform.GetUnitAxis(EAxis::Y) * RenderCuttingPlaneSize, Transform.GetLocation() + Transform.GetUnitAxis(EAxis::X) * RenderCuttingPlaneSize + Transform.GetUnitAxis(EAxis::Y) * RenderCuttingPlaneSize, FLinearColor(0, 255, 0), SDPG_Foreground);
+			PDI->DrawLine(Center + Transform.GetUnitAxis(EAxis::X) * RenderCuttingPlaneSize, Center + Transform.GetUnitAxis(EAxis::X) * RenderCuttingPlaneSize + Transform.GetUnitAxis(EAxis::Y) * RenderCuttingPlaneSize, FLinearColor(255, 0, 0), SDPG_Foreground);
+			PDI->DrawLine(Center + Transform.GetUnitAxis(EAxis::Y) * RenderCuttingPlaneSize, Center + Transform.GetUnitAxis(EAxis::X) * RenderCuttingPlaneSize + Transform.GetUnitAxis(EAxis::Y) * RenderCuttingPlaneSize, FLinearColor(0, 255, 0), SDPG_Foreground);
+		});
 	}
 }
-}
 
 void UFractureToolSlice::FractureContextChanged()
 {
 	UpdateDefaultRandomSeed();
 	TArray<FFractureToolContext> FractureContexts = GetFractureToolContexts();
 
-	RenderCuttingPlanesTransforms.Empty();
+	ClearVisualizations();
 
 	RenderCuttingPlaneSize = FLT_MAX;
 	for (FFractureToolContext& FractureContext : FractureContexts)
 	{
-		// Move the local bounds to the actor so we we'll draw in the correct location
 		FBox Bounds = FractureContext.GetWorldBounds();
+		if (!Bounds.IsValid)
+		{
+			continue;
+		}
+		int32 CollectionIdx = VisualizedCollections.Add(FractureContext.GetGeometryCollectionComponent());
+		int32 BoneIdx = FractureContext.GetSelection().Num() == 1 ? FractureContext.GetSelection()[0] : INDEX_NONE;
+		PlanesMappings.AddMapping(CollectionIdx, BoneIdx, RenderCuttingPlanesTransforms.Num());
+
 		GenerateSliceTransforms(FractureContext, RenderCuttingPlanesTransforms);
 
 		if (Bounds.GetExtent().GetMax() < RenderCuttingPlaneSize)
diff --git a/Engine/Plugins/Experimental/ChaosEditor/Source/FractureEditor/Private/FractureToolSlice.h b/Engine/Plugins/Experimental/ChaosEditor/Source/FractureEditor/Private/FractureToolSlice.h
--- a/Engine/Plugins/Experimental/ChaosEditor/Source/FractureEditor/Private/FractureToolSlice.h
+++ b/Engine/Plugins/Experimental/ChaosEditor/Source/FractureEditor/Private/FractureToolSlice.h
@@ -73,7 +73,16 @@ public:
 	UPROPERTY(EditAnywhere, Category = Slicing)
 	TObjectPtr<UFractureSliceSettings> SliceSettings;
 
+protected:
+	virtual void ClearVisualizations() override
+	{
+		Super::ClearVisualizations();
+		RenderCuttingPlanesTransforms.Empty();
+		PlanesMappings.Empty();
+	}
+
 private:
 	float RenderCuttingPlaneSize;
 	TArray<FTransform> RenderCuttingPlanesTransforms;
+	FVisualizationMappings PlanesMappings;
 };
\ No newline at end of file
diff --git a/Engine/Plugins/Experimental/ChaosEditor/Source/FractureEditor/Public/FractureTool.cpp b/Engine/Plugins/Experimental/ChaosEditor/Source/FractureEditor/Public/FractureTool.cpp
--- a/Engine/Plugins/Experimental/ChaosEditor/Source/FractureEditor/Public/FractureTool.cpp
+++ b/Engine/Plugins/Experimental/ChaosEditor/Source/FractureEditor/Public/FractureTool.cpp
@@ -195,6 +195,21 @@ TArray<FFractureToolContext> UFractureActionTool::GetFractureToolContexts() cons
 }
 
 
+void UFractureModalTool::EnumerateVisualizationMapping(const FVisualizationMappings& Mappings, int32 ArrayNum, TFunctionRef<void(int32 Idx, FVector ExplodedVector)> Func) const
+{
+	for (int32 MappingIdx = 0; MappingIdx < Mappings.Mappings.Num(); MappingIdx++)
+	{
+		const FVisualizationMappings::FIndexMapping& Mapping = Mappings.Mappings[MappingIdx];
+		FVector Offset = Mappings.GetExplodedVector(MappingIdx, VisualizedCollections[Mapping.CollectionIdx]);
+		int32 EndIdx = Mappings.GetEndIdx(MappingIdx, ArrayNum);
+		for (int32 Idx = Mapping.StartIdx; Idx < EndIdx; Idx++)
+		{
+			Func(Idx, Offset);
+		}
+	}
+}
+
+
 void UFractureModalTool::Execute(TWeakPtr<FFractureEditorModeToolkit> InToolkit)
 {
 	if (InToolkit.IsValid())
@@ -288,3 +303,19 @@ void UFractureModalTool::OnComponentTransformChangedInternal(USceneComponent* In
 		}
 	}
 }
+
+
+FVector FVisualizationMappings::GetExplodedVector(int32 MappingIdx, const UGeometryCollectionComponent* CollectionComponent) const
+{
+	int32 BoneIdx = Mappings[MappingIdx].BoneIdx;
+	if (BoneIdx != INDEX_NONE && CollectionComponent)
+	{
+		const FGeometryCollection& Collection = *CollectionComponent->GetRestCollection()->GetGeometryCollection();
+		if (Collection.HasAttribute("ExplodedVector", FGeometryCollection::TransformGroup))
+		{
+			FVector Offset = Collection.GetAttribute<FVector3f>("ExplodedVector", FGeometryCollection::TransformGroup)[BoneIdx];
+			return CollectionComponent->GetOwner()->GetActorTransform().TransformVector(Offset);
+		}
+	}
+	return FVector::ZeroVector;
+}
diff --git a/Engine/Plugins/Experimental/ChaosEditor/Source/FractureEditor/Public/FractureTool.h b/Engine/Plugins/Experimental/ChaosEditor/Source/FractureEditor/Public/FractureTool.h
--- a/Engine/Plugins/Experimental/ChaosEditor/Source/FractureEditor/Public/FractureTool.h
+++ b/Engine/Plugins/Experimental/ChaosEditor/Source/FractureEditor/Public/FractureTool.h
@@ -80,6 +80,42 @@ protected:
 
 };
 
+// Helper structure to correspond a UFractureModalTool's VisualizedCollections with elements of the tool's visualizations
+struct FVisualizationMappings
+{
+	// helper structure to map visualized geometry to source geometry and the bone that it corresponds to (if applicable)
+	struct FIndexMapping
+	{
+		int32 CollectionIdx = INDEX_NONE; // index into VisualizedCollections
+		int32 BoneIdx = INDEX_NONE; // transform (bone) index in geometry collection that corresponds to the visualization
+		int32 StartIdx = INDEX_NONE; // index of first element of visualization-related array that uses this map (e.g. index into VoronoiSites)
+	};
+
+	TArray<FIndexMapping> Mappings;
+
+	FVector GetExplodedVector(int32 MappingIdx, const UGeometryCollectionComponent* Collection) const;
+
+	void Empty()
+	{
+		Mappings.Empty();
+	}
+
+	void AddMapping(int32 CollectionIdx, int32 BoneIdx, int32 StartIdx = INDEX_NONE)
+	{
+		Mappings.Add({ CollectionIdx, BoneIdx, StartIdx });
+	}
+	int32 GetEndIdx(int32 CorrespondIdx, int32 ArrayNum) const
+	{
+		if (CorrespondIdx + 1 < Mappings.Num())
+		{
+			return Mappings[CorrespondIdx + 1].StartIdx;
+		}
+		else
+		{
+			return ArrayNum;
+		}
+	}
+};
 
 /** Tools derived from this class provide parameter details and operate modally. */
 UCLASS(Abstract)
@@ -122,6 +158,7 @@ public:
 	virtual void Shutdown() 
 	{
 		GEngine->OnComponentTransformChanged().RemoveAll(this);
+		ClearVisualizations();
 	}
 
 	// Called when a selected geometry collection component is moved in the scene
@@ -131,6 +168,17 @@ public:
 	}
 
 protected:
+
+	// Geometry collection components referenced by visualizations
+	UPROPERTY()
+	TArray<TObjectPtr<const UGeometryCollectionComponent>> VisualizedCollections;
+
+	virtual void ClearVisualizations()
+	{
+		VisualizedCollections.Empty();
+	}
+
+	void EnumerateVisualizationMapping(const FVisualizationMappings& Mappings, int32 ArrayNum, TFunctionRef<void(int32 Idx, FVector ExplodedVector)> Func) const;
 	
 	void OnComponentTransformChangedInternal(USceneComponent* InRootComponent, ETeleportType Teleport);
 
