improve undo behavior for "view settings" on fracture tools
 - stop 'explode amount' from creating non-functional undo transactions
 - hide and explode state is re-generated on undo, so we are less likely to be stuck with stale values
 - fix stale selection highlight lingering on fracture level changes

#rnx

diff --git a/Engine/Plugins/Experimental/ChaosEditor/Source/FractureEditor/Private/FractureEditorMode.cpp b/Engine/Plugins/Experimental/ChaosEditor/Source/FractureEditor/Private/FractureEditorMode.cpp
--- a/Engine/Plugins/Experimental/ChaosEditor/Source/FractureEditor/Private/FractureEditorMode.cpp
+++ b/Engine/Plugins/Experimental/ChaosEditor/Source/FractureEditor/Private/FractureEditorMode.cpp
@@ -71,7 +71,7 @@ void UFractureEditorMode::Enter()
 	if (FFractureEditorModeToolkit* FractureToolkit = static_cast<FFractureEditorModeToolkit*>(Toolkit.Get()))
 	{
 		FractureToolkit->SetInitialPalette();
-		FractureToolkit->OnHideUnselectedChanged(FractureToolkit->GetHideUnselectedValue());
+		FractureToolkit->OnHideUnselectedChanged();
 	}
 	
 	OnActorSelectionChanged(SelectedObjects, false);
diff --git a/Engine/Plugins/Experimental/ChaosEditor/Source/FractureEditor/Private/FractureEditorModeToolkit.cpp b/Engine/Plugins/Experimental/ChaosEditor/Source/FractureEditor/Private/FractureEditorModeToolkit.cpp
--- a/Engine/Plugins/Experimental/ChaosEditor/Source/FractureEditor/Private/FractureEditorModeToolkit.cpp
+++ b/Engine/Plugins/Experimental/ChaosEditor/Source/FractureEditor/Private/FractureEditorModeToolkit.cpp
@@ -471,7 +471,7 @@ void FFractureEditorModeToolkit::OnObjectPostEditChange( UObject* Object, FPrope
 		}
 		if (PropertyChangedEvent.Property->GetFName() == GET_MEMBER_NAME_CHECKED(UFractureSettings, bHideUnselected))
 		{
-			OnHideUnselectedChanged(GetHideUnselectedValue());
+			OnHideUnselectedChanged();
 		}
 		else if (PropertyChangedEvent.Property->GetFName() == GET_MEMBER_NAME_CHECKED(UHistogramSettings, bSorted))
 		{
@@ -780,8 +780,6 @@ void FFractureEditorModeToolkit::OnExplodedViewValueChanged()
 			{
 				if (UGeometryCollectionComponent* GeometryCollectionComponent = Cast<UGeometryCollectionComponent>(PrimitiveComponent))
 				{
-					FGeometryCollectionEdit RestCollection = GeometryCollectionComponent->EditRestCollection();
-					UGeometryCollection* GeometryCollection = RestCollection.GetRestCollection();
 
 					UpdateExplodedVectors(GeometryCollectionComponent);
 
@@ -861,6 +859,7 @@ void FFractureEditorModeToolkit::OnLevelViewValueChanged()
 		{
 			EditBoneColor.SetLevelViewMode(FractureLevel);
 			EditBoneColor.ResetBoneSelection();
+			EditBoneColor.ResetHighlightedBones();
 			UpdateExplodedVectors(Comp);
 			Comp->MarkRenderStateDirty();
 			Comp->MarkRenderDynamicDataDirty();
@@ -871,7 +870,33 @@ void FFractureEditorModeToolkit::OnLevelViewValueChanged()
 	GCurrentLevelEditingViewportClient->Invalidate();
 }
 
-void FFractureEditorModeToolkit::OnHideUnselectedChanged(bool bHide)
+void FFractureEditorModeToolkit::UpdateHideForComponent(UGeometryCollectionComponent* Comp)
+{
+	if (const UGeometryCollection* RestCollection = Comp->GetRestCollection())
+	{
+		TSharedPtr<FGeometryCollection, ESPMode::ThreadSafe> GeometryCollection = RestCollection->GetGeometryCollection();
+
+		if (GetHideUnselectedValue())
+		{
+			// If we are toggling on, add and configure the Hide array.
+			if (!GeometryCollection->HasAttribute("Hide", FGeometryCollection::TransformGroup))
+			{
+				GeometryCollection->AddAttribute<bool>("Hide", FGeometryCollection::TransformGroup);
+			}
+			SetHideForUnselected(Comp);
+		}
+		else
+		{
+			// If we are toggling off, remove the Hide array.
+			if (GeometryCollection->HasAttribute("Hide", FGeometryCollection::TransformGroup))
+			{
+				GeometryCollection->RemoveAttribute("Hide", FGeometryCollection::TransformGroup);
+			}
+		}
+	}
+}
+
+void FFractureEditorModeToolkit::OnHideUnselectedChanged()
 {
 	TSet<UGeometryCollectionComponent*> GeomCompSelection;
 	GetSelectedGeometryCollectionComponents(GeomCompSelection);
@@ -880,25 +905,7 @@ void FFractureEditorModeToolkit::OnHideUnselectedChanged(bool bHide)
 	{
 		if (const UGeometryCollection* RestCollection = Comp->GetRestCollection())
 		{ 
-			TSharedPtr<FGeometryCollection, ESPMode::ThreadSafe> GeometryCollection = RestCollection->GetGeometryCollection();
-			
-			if (bHide)
-			{
-				// If we are toggling on, add and configure the Hide array.
-				if (!GeometryCollection->HasAttribute("Hide", FGeometryCollection::TransformGroup))
-				{
-					GeometryCollection->AddAttribute<bool>("Hide", FGeometryCollection::TransformGroup);
-				}
-				SetHideForUnselected(Comp);
-			}
-			else
-			{
-				// If we are toggling off, remove the Hide array.
-				if (GeometryCollection->HasAttribute("Hide", FGeometryCollection::TransformGroup))
-				{
-					GeometryCollection->RemoveAttribute("Hide", FGeometryCollection::TransformGroup);
-				}
-			}
+			UpdateHideForComponent(Comp);
 		
 			// redraw
 			Comp->MarkRenderDynamicDataDirty();
@@ -1080,9 +1087,12 @@ void FFractureEditorModeToolkit::SetOutlinerComponents(const TArray<UGeometryCol
 
 			FGeometryCollectionClusteringUtility::UpdateHierarchyLevelOfChildren(GeometryCollectionPtr.Get(), -1);
 			UpdateExplodedVectors(Component);
+			UpdateHideForComponent(Component);
 
 			UpdateGeometryComponentAttributes(Component);
 			ComponentsToEdit.Add(Component);
+
+			Component->MarkRenderStateDirty();
 		}	
 	}
 
@@ -1108,7 +1118,7 @@ void FFractureEditorModeToolkit::SetBoneSelection(UGeometryCollectionComponent*
 	OutlinerView->SetBoneSelection(InRootComponent, InSelectedBones, bClearCurrentSelection);
 	HistogramView->SetBoneSelection(InRootComponent, InSelectedBones, bClearCurrentSelection);
 
-	SetHideForUnselected(InRootComponent);
+	UpdateHideForComponent(InRootComponent);
 	
 	if (ActiveTool != nullptr)
 	{
@@ -1403,14 +1413,16 @@ void FFractureEditorModeToolkit::UpdateExplodedVectors(UGeometryCollectionCompon
 	{
 		if (OutGeometryCollectionConst->HasAttribute("ExplodedVector", FGeometryCollection::TransformGroup))
 		{
-			FGeometryCollectionEdit RestCollection = GeometryCollectionComponent->EditRestCollection(GeometryCollection::EEditUpdate::RestPhysicsDynamic);
+			// ExplodedVector is not saved, so the Rest collection doesn't 'see' this update in serialization, so we don't need EEditUpdate::Rest here
+			FGeometryCollectionEdit RestCollection = GeometryCollectionComponent->EditRestCollection(GeometryCollection::EEditUpdate::Dynamic);
 			FGeometryCollection* OutGeometryCollection = GeometryCollectionPtr.Get();
 			OutGeometryCollection->RemoveAttribute("ExplodedVector", FGeometryCollection::TransformGroup);
 		}
 	}
 	else
 	{
-		FGeometryCollectionEdit RestCollection = GeometryCollectionComponent->EditRestCollection(GeometryCollection::EEditUpdate::RestPhysicsDynamic);
+		// ExplodedVector is not saved, so the Rest collection doesn't 'see' this update in serialization, so we don't need EEditUpdate::Rest here
+		FGeometryCollectionEdit RestCollection = GeometryCollectionComponent->EditRestCollection(GeometryCollection::EEditUpdate::Dynamic);
 		UGeometryCollection* GeometryCollection = RestCollection.GetRestCollection();
 		FGeometryCollection* OutGeometryCollection = GeometryCollectionPtr.Get();
 
@@ -1523,7 +1535,7 @@ void FFractureEditorModeToolkit::OnOutlinerBoneSelectionChanged(UGeometryCollect
 			ActiveTool->FractureContextChanged();
 		}
 
-		SetHideForUnselected(RootComponent);
+		UpdateHideForComponent(RootComponent);
 
 		RootComponent->MarkRenderStateDirty();
 		RootComponent->MarkRenderDynamicDataDirty();
@@ -1555,7 +1567,7 @@ void FFractureEditorModeToolkit::OnHistogramBoneSelectionChanged(UGeometryCollec
 			ActiveTool->FractureContextChanged();
 		}
 
-		SetHideForUnselected(RootComponent);
+		UpdateHideForComponent(RootComponent);
 
 		RootComponent->MarkRenderStateDirty();
 		RootComponent->MarkRenderDynamicDataDirty();
diff --git a/Engine/Plugins/Experimental/ChaosEditor/Source/FractureEditor/Public/FractureEditorModeToolkit.h b/Engine/Plugins/Experimental/ChaosEditor/Source/FractureEditor/Public/FractureEditorModeToolkit.h
--- a/Engine/Plugins/Experimental/ChaosEditor/Source/FractureEditor/Public/FractureEditorModeToolkit.h
+++ b/Engine/Plugins/Experimental/ChaosEditor/Source/FractureEditor/Public/FractureEditorModeToolkit.h
@@ -117,7 +117,8 @@ public:
 
 	void OnExplodedViewValueChanged();
 	void OnLevelViewValueChanged();
-	void OnHideUnselectedChanged(bool bHide);
+	void OnHideUnselectedChanged();
+	void UpdateHideForComponent(UGeometryCollectionComponent* Component);
 
 	// Update any View Property Changes 
 	void OnObjectPostEditChange( UObject* Object, FPropertyChangedEvent& PropertyChangedEvent );
