Extend fracture noise settings to include persistence and lacunarity (+ remove "Voronoi" label from non-Voronoi fracture types)

diff --git a/Engine/Plugins/Experimental/ChaosEditor/Source/FractureEditor/Private/FractureToolBrick.cpp b/Engine/Plugins/Experimental/ChaosEditor/Source/FractureEditor/Private/FractureToolBrick.cpp
--- a/Engine/Plugins/Experimental/ChaosEditor/Source/FractureEditor/Private/FractureToolBrick.cpp
+++ b/Engine/Plugins/Experimental/ChaosEditor/Source/FractureEditor/Private/FractureToolBrick.cpp
@@ -40,7 +40,7 @@ FSlateIcon UFractureToolBrick::GetToolIcon() const
 
 void UFractureToolBrick::RegisterUICommand( FFractureEditorCommands* BindingContext ) 
 {
-	UI_COMMAND_EXT( BindingContext, UICommandInfo, "Brick", "Brick", "Brick Voronoi Fracture", EUserInterfaceActionType::ToggleButton, FInputChord() );
+	UI_COMMAND_EXT( BindingContext, UICommandInfo, "Brick", "Brick", "Brick Fracture", EUserInterfaceActionType::ToggleButton, FInputChord() );
 	BindingContext->Brick = UICommandInfo;
 }
 
@@ -136,7 +136,7 @@ namespace FractureToolBrickLocals
 			Dimensions = SafeDimensions;
 			NumBricks = SafeNumBricks;
 
-			UE_LOG(LogFractureTool, Warning, TEXT("Brick Voronoi Fracture: Current brick dimensions of %f x %f x %f would result in %d bricks. "
+			UE_LOG(LogFractureTool, Warning, TEXT("Brick Fracture: Current brick dimensions of %f x %f x %f would result in %d bricks. "
 				"Reduced brick dimensions to %f x %f x %f resulting in %d bricks to stay within maximum number of %d bricks."),
 				BrickSettings->BrickLength, BrickSettings->BrickDepth, BrickSettings->BrickHeight, InputNumBricks,
 				Dimensions.X, Dimensions.Y, Dimensions.Z, NumBricks, NumBricksLimit);
@@ -407,10 +407,7 @@ int32 UFractureToolBrick::ExecuteFracture(const FFractureToolContext& FractureCo
 		FNoiseSettings NoiseSettings;
 		if (CutterSettings->Amplitude > 0.0f)
 		{
-			NoiseSettings.Amplitude = CutterSettings->Amplitude;
-			NoiseSettings.Frequency = CutterSettings->Frequency;
-			NoiseSettings.Octaves = CutterSettings->OctaveNumber;
-			NoiseSettings.PointSpacing = CutterSettings->SurfaceResolution;
+			CutterSettings->TransferNoiseSettings(NoiseSettings);
 			VoronoiPlanarCells.InternalSurfaceMaterials.NoiseSettings = NoiseSettings;
 		}
 
diff --git a/Engine/Plugins/Experimental/ChaosEditor/Source/FractureEditor/Private/FractureToolCutter.cpp b/Engine/Plugins/Experimental/ChaosEditor/Source/FractureEditor/Private/FractureToolCutter.cpp
--- a/Engine/Plugins/Experimental/ChaosEditor/Source/FractureEditor/Private/FractureToolCutter.cpp
+++ b/Engine/Plugins/Experimental/ChaosEditor/Source/FractureEditor/Private/FractureToolCutter.cpp
@@ -97,6 +97,16 @@ void UFractureTransformGizmoSettings::Shutdown()
 	}
 }
 
+void UFractureCutterSettings::TransferNoiseSettings(FNoiseSettings& NoiseSettingsOut)
+{
+	NoiseSettingsOut.Amplitude = Amplitude;
+	NoiseSettingsOut.Frequency = Frequency;
+	NoiseSettingsOut.Lacunarity = Lacunarity;
+	NoiseSettingsOut.Persistence = Persistence;
+	NoiseSettingsOut.Octaves = OctaveNumber;
+	NoiseSettingsOut.PointSpacing = SurfaceResolution;
+}
+
 UFractureToolCutterBase::UFractureToolCutterBase(const FObjectInitializer& ObjInit)
 	: Super(ObjInit)
 {
@@ -304,10 +314,7 @@ int32 UFractureToolVoronoiCutterBase::ExecuteFracture(const FFractureToolContext
 		FNoiseSettings NoiseSettings;
 		if (CutterSettings->Amplitude > 0.0f)
 		{
-			NoiseSettings.Amplitude = CutterSettings->Amplitude;
-			NoiseSettings.Frequency = CutterSettings->Frequency;
-			NoiseSettings.Octaves = CutterSettings->OctaveNumber;
-			NoiseSettings.PointSpacing = CutterSettings->SurfaceResolution;
+			CutterSettings->TransferNoiseSettings(NoiseSettings);
 			VoronoiPlanarCells.InternalSurfaceMaterials.NoiseSettings = NoiseSettings;
 		}
 
@@ -322,9 +329,13 @@ int32 UFractureToolVoronoiCutterBase::ExecuteFracture(const FFractureToolContext
 
 FBox UFractureToolVoronoiCutterBase::GetVoronoiBounds(const FFractureToolContext& FractureContext, const TArray<FVector>& Sites) const
 {
-	FBox VoronoiBounds(Sites);
-	VoronoiBounds += FractureContext.GetWorldBounds();
-	return VoronoiBounds.ExpandBy(CutterSettings->Amplitude + CutterSettings->Grout + KINDA_SMALL_NUMBER);
+	FBox VoronoiBounds = FractureContext.GetWorldBounds(); 
+	if (Sites.Num() > 0)
+	{
+		VoronoiBounds += FBox(Sites);
+	}
+	
+	return VoronoiBounds.ExpandBy(CutterSettings->GetMaxVertexMovement() + KINDA_SMALL_NUMBER);
 }
 
 #undef LOCTEXT_NAMESPACE
diff --git a/Engine/Plugins/Experimental/ChaosEditor/Source/FractureEditor/Private/FractureToolCutter.h b/Engine/Plugins/Experimental/ChaosEditor/Source/FractureEditor/Private/FractureToolCutter.h
--- a/Engine/Plugins/Experimental/ChaosEditor/Source/FractureEditor/Private/FractureToolCutter.h
+++ b/Engine/Plugins/Experimental/ChaosEditor/Source/FractureEditor/Private/FractureToolCutter.h
@@ -11,6 +11,8 @@
 
 class FFractureToolContext;
 
+struct FNoiseSettings;
+
 /** Settings specifically related to the one-time destructive fracturing of a mesh **/
 UCLASS(config = EditorPerProjectUserSettings)
 class UFractureCutterSettings : public UFractureToolSettings
@@ -28,10 +30,14 @@ public:
 		, bDrawDiagram(true)
 		, Amplitude(0.0f)
 		, Frequency(0.1f)
+		, Lacunarity(2.0f)
+		, Persistence(0.5f)
 		, OctaveNumber(4)
 		, SurfaceResolution(10)
 	{}
 
+	void TransferNoiseSettings(FNoiseSettings& NoiseSettingsOut);
+
 	/** Random number generator seed for repeatability */
 	UPROPERTY(EditAnywhere, Category = CommonFracture, meta = (DisplayName = "Random Seed", UIMin = "-1", UIMax = "1000", ClampMin = "-1"))
 	int32 RandomSeed;
@@ -68,13 +74,33 @@ public:
 	UPROPERTY(EditAnywhere, Category = Noise)
 	float Frequency;
 
+	/** Lacunarity of the Perlin noise.  Controls how the frequency scales per octave. */
+	UPROPERTY(EditAnywhere, Category = Noise, meta = (UIMin = "1.0", UIMax = "4.0"))
+	float Lacunarity;
+
+	/** Persistence of the Perlin noise.  Controls how the amplitude scales per octave. */
+	UPROPERTY(EditAnywhere, Category = Noise, meta = (UIMin = "0.01", UIMax = "1.0"))
+	float Persistence;
+
 	/** Number of fractal layers of Perlin noise to apply.  Smaller values (1 or 2) will create noise that looks like gentle rolling hills, while larger values (> 4) will tend to look more like craggy mountains */
-	UPROPERTY(EditAnywhere, Category = Noise, meta = (UIMin = "1"))
+	UPROPERTY(EditAnywhere, Category = Noise, meta = (ClampMin = "1", UIMax = "8"))
 	int32 OctaveNumber;
 
 	/** Spacing between vertices on cut surfaces, where noise is added.  Larger spacing between vertices will create more efficient meshes with fewer triangles, but less resolution to see the shape of the added noise  */
 	UPROPERTY(EditAnywhere, Category = Noise, meta = (DisplayName = "Point Spacing", UIMin = "1", ClampMin = "0.1"))
 	float SurfaceResolution;
+
+	/// Get the maximum distance a vertex could be moved by a combination of grout and noise
+	float GetMaxVertexMovement()
+	{
+		float MaxDisp = Grout;
+		float AmplitudeScaled = Amplitude;
+		for (int32 OctaveIdx = 0; OctaveIdx < OctaveNumber; OctaveIdx++, AmplitudeScaled *= Persistence)
+		{
+			MaxDisp += FMath::Abs(AmplitudeScaled);
+		}
+		return MaxDisp;
+	}
 };
 
 /** Settings related to the collision properties of the fractured mesh pieces */
diff --git a/Engine/Plugins/Experimental/ChaosEditor/Source/FractureEditor/Private/FractureToolPlaneCut.cpp b/Engine/Plugins/Experimental/ChaosEditor/Source/FractureEditor/Private/FractureToolPlaneCut.cpp
--- a/Engine/Plugins/Experimental/ChaosEditor/Source/FractureEditor/Private/FractureToolPlaneCut.cpp
+++ b/Engine/Plugins/Experimental/ChaosEditor/Source/FractureEditor/Private/FractureToolPlaneCut.cpp
@@ -51,7 +51,7 @@ FSlateIcon UFractureToolPlaneCut::GetToolIcon() const
 
 void UFractureToolPlaneCut::RegisterUICommand( FFractureEditorCommands* BindingContext ) 
 {
-	UI_COMMAND_EXT( BindingContext, UICommandInfo, "Planar", "Planar", "Planar Voronoi Fracture", EUserInterfaceActionType::ToggleButton, FInputChord() );
+	UI_COMMAND_EXT( BindingContext, UICommandInfo, "Planar", "Planar", "Planar Fracture", EUserInterfaceActionType::ToggleButton, FInputChord() );
 	BindingContext->Planar = UICommandInfo;
 }
 
@@ -163,10 +163,7 @@ int32 UFractureToolPlaneCut::ExecuteFracture(const FFractureToolContext& Fractur
 		FNoiseSettings NoiseSettings;
 		if (CutterSettings->Amplitude > 0.0f)
 		{
-			NoiseSettings.Amplitude = CutterSettings->Amplitude;
-			NoiseSettings.Frequency = CutterSettings->Frequency;
-			NoiseSettings.Octaves = CutterSettings->OctaveNumber;
-			NoiseSettings.PointSpacing = CutterSettings->SurfaceResolution;
+			CutterSettings->TransferNoiseSettings(NoiseSettings);
 			InternalSurfaceMaterials.NoiseSettings = NoiseSettings;
 		}
 
diff --git a/Engine/Plugins/Experimental/ChaosEditor/Source/FractureEditor/Private/FractureToolSlice.cpp b/Engine/Plugins/Experimental/ChaosEditor/Source/FractureEditor/Private/FractureToolSlice.cpp
--- a/Engine/Plugins/Experimental/ChaosEditor/Source/FractureEditor/Private/FractureToolSlice.cpp
+++ b/Engine/Plugins/Experimental/ChaosEditor/Source/FractureEditor/Private/FractureToolSlice.cpp
@@ -30,7 +30,7 @@ FText UFractureToolSlice::GetDisplayText() const
 
 FText UFractureToolSlice::GetTooltipText() const 
 { 
-	return FText(NSLOCTEXT("Fracture", "FractureToolSliceTooltip", "The Slice Voronoi fracture method enables you to define the number of X, Y, and Z slices, along with providing random angle and offset variation.  Click the Fracture Button to commit the fracture to the geometry collection.")); 
+	return FText(NSLOCTEXT("Fracture", "FractureToolSliceTooltip", "The Slice fracture method enables you to define the number of X, Y, and Z slices, along with providing random angle and offset variation.  Click the Fracture Button to commit the fracture to the geometry collection.")); 
 }
 
 FSlateIcon UFractureToolSlice::GetToolIcon() const 
@@ -40,7 +40,7 @@ FSlateIcon UFractureToolSlice::GetToolIcon() const
 
 void UFractureToolSlice::RegisterUICommand( FFractureEditorCommands* BindingContext ) 
 {
-	UI_COMMAND_EXT( BindingContext, UICommandInfo, "Slice", "Slice", "Slice Voronoi Fracture", EUserInterfaceActionType::ToggleButton, FInputChord() );
+	UI_COMMAND_EXT( BindingContext, UICommandInfo, "Slice", "Slice", "Slice Fracture", EUserInterfaceActionType::ToggleButton, FInputChord() );
 	BindingContext->Slice = UICommandInfo;
 }
 
@@ -176,10 +176,7 @@ int32 UFractureToolSlice::ExecuteFracture(const FFractureToolContext& FractureCo
 		FNoiseSettings NoiseSettings;
 		if (CutterSettings->Amplitude > 0.0f)
 		{
-			NoiseSettings.Amplitude = CutterSettings->Amplitude;
-			NoiseSettings.Frequency = CutterSettings->Frequency;
-			NoiseSettings.Octaves = CutterSettings->OctaveNumber;
-			NoiseSettings.PointSpacing = CutterSettings->SurfaceResolution;
+			CutterSettings->TransferNoiseSettings(NoiseSettings);
 			InternalSurfaceMaterials.NoiseSettings = NoiseSettings;
 		}
 
diff --git a/Engine/Plugins/Experimental/PlanarCutPlugin/Source/PlanarCut/Private/GeometryMeshConversion.cpp b/Engine/Plugins/Experimental/PlanarCutPlugin/Source/PlanarCut/Private/GeometryMeshConversion.cpp
--- a/Engine/Plugins/Experimental/PlanarCutPlugin/Source/PlanarCut/Private/GeometryMeshConversion.cpp
+++ b/Engine/Plugins/Experimental/PlanarCutPlugin/Source/PlanarCut/Private/GeometryMeshConversion.cpp
@@ -1179,10 +1179,11 @@ float FCellMeshes::OctaveNoise(const FVector& V, const FNoiseSettings& Settings)
 {
 	int32 Octaves = Settings.Octaves;
 	float NoiseValue = 0;
-	float OctaveScale = 1;
-	for (int32 Octave = 0; Octave < Octaves; Octave++, OctaveScale *= 2)
+	float FreqScale = 1;
+	float AmpScale = 1;
+	for (int32 Octave = 0; Octave < Octaves; Octave++, FreqScale *= Settings.Lacunarity, AmpScale *= Settings.Persistence)
 	{
-		NoiseValue += FMath::PerlinNoise3D(V * OctaveScale) / OctaveScale;
+		NoiseValue += FMath::PerlinNoise3D(V * FreqScale) * AmpScale;
 	}
 	return NoiseValue;
 }
diff --git a/Engine/Plugins/Experimental/PlanarCutPlugin/Source/PlanarCut/Public/PlanarCut.h b/Engine/Plugins/Experimental/PlanarCutPlugin/Source/PlanarCut/Public/PlanarCut.h
--- a/Engine/Plugins/Experimental/PlanarCutPlugin/Source/PlanarCut/Public/PlanarCut.h
+++ b/Engine/Plugins/Experimental/PlanarCutPlugin/Source/PlanarCut/Public/PlanarCut.h
@@ -14,6 +14,8 @@ struct PLANARCUT_API FNoiseSettings
 	float Frequency = .1;
 	int32 Octaves = 4;
 	float PointSpacing = 1;
+	float Lacunarity = 2;
+	float Persistence = .5;
 };
 
 // auxiliary structure for FPlanarCells to carry material info
