toggle fracture plane cut tool's 'number of plane cuts' off when gizmo is enabled, because plane will by defined by gizmo in that case
and add related machinery to support toggling settings across different settings objects

#rnx

diff --git a/Engine/Plugins/Experimental/ChaosEditor/Source/FractureEditor/Private/FractureEditorModeToolkit.cpp b/Engine/Plugins/Experimental/ChaosEditor/Source/FractureEditor/Private/FractureEditorModeToolkit.cpp
--- a/Engine/Plugins/Experimental/ChaosEditor/Source/FractureEditor/Private/FractureEditorModeToolkit.cpp
+++ b/Engine/Plugins/Experimental/ChaosEditor/Source/FractureEditor/Private/FractureEditorModeToolkit.cpp
@@ -1000,12 +1000,18 @@ bool FFractureEditorModeToolkit::CanExecuteAction(UFractureActionTool* InActionT
 	}
 }
 
-void FFractureEditorModeToolkit::SetActiveTool(UFractureModalTool* InActiveTool)
+void FFractureEditorModeToolkit::ShutdownActiveTool()
 {
 	if (ActiveTool)
 	{
 		ActiveTool->Shutdown();
+		ActiveTool->OnPropertyModifiedDirectlyByTool.RemoveAll(this);
 	}
+}
+
+void FFractureEditorModeToolkit::SetActiveTool(UFractureModalTool* InActiveTool)
+{
+	ShutdownActiveTool();
 
 	ActiveTool = InActiveTool;
 
@@ -1017,6 +1023,8 @@ void FFractureEditorModeToolkit::SetActiveTool(UFractureModalTool* InActiveTool)
 
 	if (ActiveTool != nullptr)
 	{
+		ActiveTool->OnPropertyModifiedDirectlyByTool.AddSP(this, &FFractureEditorModeToolkit::InvalidateCachedDetailPanelState);
+
 		ActiveTool->Setup();
 
 		Settings.Append(ActiveTool->GetSettingsObjects());
@@ -1028,12 +1036,14 @@ void FFractureEditorModeToolkit::SetActiveTool(UFractureModalTool* InActiveTool)
 	DetailsView->SetObjects(Settings);
 }
 
+void FFractureEditorModeToolkit::InvalidateCachedDetailPanelState(UObject* ChangedObject)
+{
+	DetailsView->InvalidateCachedState();
+}
+
 void FFractureEditorModeToolkit::Shutdown()
 {
-	if (ActiveTool)
-	{
-		ActiveTool->Shutdown();
-	}
+	ShutdownActiveTool();
 }
 
 
diff --git a/Engine/Plugins/Experimental/ChaosEditor/Source/FractureEditor/Private/FractureToolCutter.h b/Engine/Plugins/Experimental/ChaosEditor/Source/FractureEditor/Private/FractureToolCutter.h
--- a/Engine/Plugins/Experimental/ChaosEditor/Source/FractureEditor/Private/FractureToolCutter.h
+++ b/Engine/Plugins/Experimental/ChaosEditor/Source/FractureEditor/Private/FractureToolCutter.h
@@ -235,13 +235,16 @@ public:
 	}
 
 	/** Use a 3D rigid transform gizmo to place the fracture pattern.  Only supports grouped fracture. */
-	UPROPERTY(EditAnywhere, Category = PlacementControls)
+	UPROPERTY(EditAnywhere, Category = PlacementControls, meta = (EditCondition = "bShowUseGizmoOption", EditConditionHides, HideEditConditionToggle))
 	bool bUseGizmo = true;
 
 	/** Recenter the gizmo to the center of the selection when selection changes */
-	UPROPERTY(EditAnywhere, Category = PlacementControls, meta = (EditCondition = "bUseGizmo"))
+	UPROPERTY(EditAnywhere, Category = PlacementControls, meta = (EditCondition = "bUseGizmo", EditConditionHides))
 	bool bCenterOnSelection = true;
 
+	UPROPERTY()
+	bool bShowUseGizmoOption = true;
+
 	UPROPERTY()
 	TObjectPtr<UTransformGizmo> TransformGizmo = nullptr;
 
diff --git a/Engine/Plugins/Experimental/ChaosEditor/Source/FractureEditor/Private/FractureToolPlaneCut.cpp b/Engine/Plugins/Experimental/ChaosEditor/Source/FractureEditor/Private/FractureToolPlaneCut.cpp
--- a/Engine/Plugins/Experimental/ChaosEditor/Source/FractureEditor/Private/FractureToolPlaneCut.cpp
+++ b/Engine/Plugins/Experimental/ChaosEditor/Source/FractureEditor/Private/FractureToolPlaneCut.cpp
@@ -25,6 +25,8 @@ UFractureToolPlaneCut::UFractureToolPlaneCut(const FObjectInitializer& ObjInit)
 void UFractureToolPlaneCut::Setup()
 {
 	GizmoSettings->Setup(this);
+	PlaneCutSettings->bCanCutWithMultiplePlanes = !GizmoSettings->bUseGizmo;
+	NotifyOfPropertyChangeByTool(PlaneCutSettings);
 }
 
 
@@ -109,6 +111,19 @@ TArray<UObject*> UFractureToolPlaneCut::GetSettingsObjects() const
 	return Settings;
 }
 
+
+void UFractureToolPlaneCut::PostEditChangeChainProperty(FPropertyChangedChainEvent& PropertyChangedEvent)
+{
+	if (PropertyChangedEvent.Property->GetFName() == GET_MEMBER_NAME_CHECKED(UFractureTransformGizmoSettings, bUseGizmo))
+	{
+		PlaneCutSettings->bCanCutWithMultiplePlanes = !GizmoSettings->bUseGizmo;
+		NotifyOfPropertyChangeByTool(PlaneCutSettings);
+	}
+	
+	Super::PostEditChangeChainProperty(PropertyChangedEvent);
+}
+
+
 void UFractureToolPlaneCut::FractureContextChanged()
 {
 	UpdateDefaultRandomSeed();
diff --git a/Engine/Plugins/Experimental/ChaosEditor/Source/FractureEditor/Private/FractureToolPlaneCut.h b/Engine/Plugins/Experimental/ChaosEditor/Source/FractureEditor/Private/FractureToolPlaneCut.h
--- a/Engine/Plugins/Experimental/ChaosEditor/Source/FractureEditor/Private/FractureToolPlaneCut.h
+++ b/Engine/Plugins/Experimental/ChaosEditor/Source/FractureEditor/Private/FractureToolPlaneCut.h
@@ -19,9 +19,12 @@ public:
 		: Super(ObjInit)
 		, NumberPlanarCuts(1) {}
 
-	/** Number of Cutting Planes */
-	UPROPERTY(EditAnywhere, Category = PlaneCut, meta = (DisplayName = "Number of Cuts", UIMin = "1", UIMax = "20", ClampMin = "1", EditCondition = "bCanCutWithMultiplePlanes"))
+	/** Number of Cutting Planes (Note: Only applicable if interactive placement gizmo is not in use) */
+	UPROPERTY(EditAnywhere, Category = PlaneCut, meta = (DisplayName = "Number of Cuts", UIMin = "1", UIMax = "20", ClampMin = "1", EditCondition = "bCanCutWithMultiplePlanes", HideEditConditionToggle))
 	int32 NumberPlanarCuts;
+
+	UPROPERTY()
+	bool bCanCutWithMultiplePlanes = false;
 };
 
 
@@ -59,6 +62,8 @@ protected:
 		PlanesMappings.Empty();
 	}
 
+	virtual void PostEditChangeChainProperty(struct FPropertyChangedChainEvent& PropertyChangedEvent) override;
+
 private:
 	// Slicing
 	UPROPERTY(EditAnywhere, Category = Slicing)
diff --git a/Engine/Plugins/Experimental/ChaosEditor/Source/FractureEditor/Public/FractureEditorModeToolkit.h b/Engine/Plugins/Experimental/ChaosEditor/Source/FractureEditor/Public/FractureEditorModeToolkit.h
--- a/Engine/Plugins/Experimental/ChaosEditor/Source/FractureEditor/Public/FractureEditorModeToolkit.h
+++ b/Engine/Plugins/Experimental/ChaosEditor/Source/FractureEditor/Public/FractureEditorModeToolkit.h
@@ -102,6 +102,7 @@ public:
 	void SetActiveTool(UFractureModalTool* InActiveTool);
 	UFractureModalTool* GetActiveTool() const;
 	bool IsActiveTool(UFractureModalTool* InActiveTool);
+	void ShutdownActiveTool();
 
 	void Shutdown();
 
@@ -181,6 +182,8 @@ protected:
 	static bool IsGeometryCollectionSelected();
 	static bool IsSelectedActorsInEditorWorld();	
 
+	void InvalidateCachedDetailPanelState(UObject* ChangedObject);
+
 private:
 	static void UpdateGeometryComponentAttributes(UGeometryCollectionComponent* Component);
 	static void UpdateVolumes(FGeometryCollectionPtr GeometryCollection, const Chaos::TParticles<Chaos::FReal, 3>& MassSpaceParticles, int32 TransformIndex);
diff --git a/Engine/Plugins/Experimental/ChaosEditor/Source/FractureEditor/Public/FractureTool.cpp b/Engine/Plugins/Experimental/ChaosEditor/Source/FractureEditor/Public/FractureTool.cpp
--- a/Engine/Plugins/Experimental/ChaosEditor/Source/FractureEditor/Public/FractureTool.cpp
+++ b/Engine/Plugins/Experimental/ChaosEditor/Source/FractureEditor/Public/FractureTool.cpp
@@ -210,6 +210,12 @@ void UFractureModalTool::EnumerateVisualizationMapping(const FVisualizationMappi
 }
 
 
+void UFractureModalTool::NotifyOfPropertyChangeByTool(UFractureToolSettings* PropertySet) const
+{
+	OnPropertyModifiedDirectlyByTool.Broadcast(PropertySet);
+}
+
+
 void UFractureModalTool::Execute(TWeakPtr<FFractureEditorModeToolkit> InToolkit)
 {
 	if (InToolkit.IsValid())
diff --git a/Engine/Plugins/Experimental/ChaosEditor/Source/FractureEditor/Public/FractureTool.h b/Engine/Plugins/Experimental/ChaosEditor/Source/FractureEditor/Public/FractureTool.h
--- a/Engine/Plugins/Experimental/ChaosEditor/Source/FractureEditor/Public/FractureTool.h
+++ b/Engine/Plugins/Experimental/ChaosEditor/Source/FractureEditor/Public/FractureTool.h
@@ -167,6 +167,22 @@ public:
 		FractureContextChanged();
 	}
 
+	/**
+	 * Call after changing properties internally in the tool to allow external views of the property
+	 * to update properly. This is meant as an outward notification mechanism, not a way to to
+	 * pass along notifications, so don't call this if the property is changed externally (i.e., this
+	 * should not usually be called from OnPropertyModified unless the tool adds changes of its own).
+	 */
+	virtual void NotifyOfPropertyChangeByTool(UFractureToolSettings* PropertySet) const;
+
+	/**
+	 * OnPropertyModifiedDirectlyByTool is broadcast when a property is changed internally by the tool.
+	 * This allows any external display of such properties to update. In a DetailsView, for instance,
+	 * it refreshes certain cached states such as edit condition states for other properties.
+	 */
+	DECLARE_MULTICAST_DELEGATE_OneParam(OnToolPropertyInternallyModified, UObject*);
+	OnToolPropertyInternallyModified OnPropertyModifiedDirectlyByTool;
+
 protected:
 
 	// Geometry collection components referenced by visualizations
