Fracture Mode
-- Fixed math error in Relative Size calculation

diff --git a/Engine/Plugins/Experimental/ChaosEditor/Source/FractureEditor/Private/FractureEditorModeToolkit.cpp b/Engine/Plugins/Experimental/ChaosEditor/Source/FractureEditor/Private/FractureEditorModeToolkit.cpp
--- a/Engine/Plugins/Experimental/ChaosEditor/Source/FractureEditor/Private/FractureEditorModeToolkit.cpp
+++ b/Engine/Plugins/Experimental/ChaosEditor/Source/FractureEditor/Private/FractureEditorModeToolkit.cpp
@@ -1280,6 +1280,9 @@ void FFractureEditorModeToolkit::UpdateGeometryComponentAttributes(UGeometryColl
 			float MaxSize = 0.0f;
 			for (int32 Idx = 0; Idx < Size.Num(); ++Idx)
 			{
+				// We use cube root because it's a more manageable way to look at the numbers
+				Size[Idx] = FGenericPlatformMath::Pow(static_cast<float>(Size[Idx]), (1.0f / 3.0f));
+				
 				if (SimulationType[Idx] == FGeometryCollection::ESimulationTypes::FST_Rigid)
 				{ 
 					MaxSize = FMath::Max(Size[Idx], MaxSize);
@@ -1334,7 +1337,7 @@ void FFractureEditorModeToolkit::UpdateVolumes(FGeometryCollectionPtr GeometryCo
 
 			// Since we're only interested in relative mass, we assume density = 1.0
 			Volumes[TransformIndex] = Volume; // todo(lwc) potential conversion from double to float
-			Size[TransformIndex] = FGenericPlatformMath::Pow(static_cast<float>(Volume), (1.0f/3.0f));
+			Size[TransformIndex] = Volume;
 		}	
 	}
 	else if (SimulationType[TransformIndex] == FGeometryCollection::ESimulationTypes::FST_Clustered)
