Add option for resample tool to choose whether it visualizes all vertices or only the added (floating) vertices

diff --git a/Engine/Plugins/Experimental/ChaosEditor/Source/FractureEditor/Private/FractureToolResample.cpp b/Engine/Plugins/Experimental/ChaosEditor/Source/FractureEditor/Private/FractureToolResample.cpp
--- a/Engine/Plugins/Experimental/ChaosEditor/Source/FractureEditor/Private/FractureToolResample.cpp
+++ b/Engine/Plugins/Experimental/ChaosEditor/Source/FractureEditor/Private/FractureToolResample.cpp
@@ -15,6 +15,8 @@
 UFractureToolResample::UFractureToolResample(const FObjectInitializer& ObjInit) 
 	: Super(ObjInit) 
 {
+	ResampleSettings = NewObject<UFractureResampleSettings>(GetTransientPackage(), UFractureResampleSettings::StaticClass());
+	ResampleSettings->OwnerTool = this;
 }
 
 FText UFractureToolResample::GetDisplayText() const
@@ -51,6 +53,7 @@ TArray<UObject*> UFractureToolResample::GetSettingsObjects() const
  {
 	TArray<UObject*> Settings;
 	Settings.Add(CollisionSettings);
+	Settings.Add(ResampleSettings);
 	return Settings;
 }
 
@@ -81,10 +84,13 @@ void UFractureToolResample::FractureContextChanged()
 			int32 FaceStart = Collection.FaceStart[GeometryIdx];
 			int32 FaceEnd = FaceStart + Collection.FaceCount[GeometryIdx];
 			// only show off-vertex samples; skip over the samples that are on faces
-			for (int32 FIdx = FaceStart; FIdx < FaceEnd; FIdx++)
+			if (ResampleSettings->bOnlyShowAddedPoints)
 			{
-				FIntVector Face = Collection.Indices[FIdx];
-				VertStart = FMath::Max(VertStart, Face.GetMax() + 1);
+				for (int32 FIdx = FaceStart; FIdx < FaceEnd; FIdx++)
+				{
+					FIntVector Face = Collection.Indices[FIdx];
+					VertStart = FMath::Max(VertStart, Face.GetMax() + 1);
+				}
 			}
 			for (int32 VIdx = VertStart; VIdx < VertEnd; VIdx++)
 			{
diff --git a/Engine/Plugins/Experimental/ChaosEditor/Source/FractureEditor/Private/FractureToolResample.h b/Engine/Plugins/Experimental/ChaosEditor/Source/FractureEditor/Private/FractureToolResample.h
--- a/Engine/Plugins/Experimental/ChaosEditor/Source/FractureEditor/Private/FractureToolResample.h
+++ b/Engine/Plugins/Experimental/ChaosEditor/Source/FractureEditor/Private/FractureToolResample.h
@@ -9,6 +9,24 @@
 
 class FFractureToolContext;
 
+/** Settings giving additional control over geometry resampling + related visualization */
+UCLASS(config = EditorPerProjectUserSettings)
+class UFractureResampleSettings : public UFractureToolSettings
+{
+public:
+
+	GENERATED_BODY()
+
+	UFractureResampleSettings(const FObjectInitializer& ObjInit)
+		: Super(ObjInit)
+	{}
+
+	/** Whether to visualize all mesh vertices or only the ones added by resampling */
+	UPROPERTY(EditAnywhere, Category = FilterSettings)
+	bool bOnlyShowAddedPoints = false;
+
+};
+
 // Note this tool doesn't actually fracture, but it does remake pieces of geometry and shares a lot of machinery with the fracture tools
 UCLASS(DisplayName = "Resample Tool", Category = "FractureTools")
 class UFractureToolResample : public UFractureToolCutterBase
@@ -44,6 +62,10 @@ protected:
 
 private:
 
+
+	UPROPERTY(EditAnywhere, Category = FixGeo)
+	TObjectPtr<UFractureResampleSettings> ResampleSettings;
+
 	TArray<FVector> GeneratedPoints;
 	FVisualizationMappings PointsMappings;
 };
