FractureMode
-- exit mode on map change

diff --git a/Engine/Plugins/Experimental/ChaosEditor/Source/FractureEditor/Private/FractureEditorModeToolkit.cpp b/Engine/Plugins/Experimental/ChaosEditor/Source/FractureEditor/Private/FractureEditorModeToolkit.cpp
--- a/Engine/Plugins/Experimental/ChaosEditor/Source/FractureEditor/Private/FractureEditorModeToolkit.cpp
+++ b/Engine/Plugins/Experimental/ChaosEditor/Source/FractureEditor/Private/FractureEditorModeToolkit.cpp
@@ -63,6 +63,8 @@
 #include "DetailLayoutBuilder.h"
 #include "PropertyHandle.h"
 
+#include "LevelEditor.h"
+
 #include "FractureSettings.h"
 #include "Toolkits/AssetEditorModeUILayer.h"
 #include "SPrimaryButton.h"
@@ -183,6 +185,12 @@ FFractureEditorModeToolkit::FFractureEditorModeToolkit()
 FFractureEditorModeToolkit::~FFractureEditorModeToolkit()
 {
 	FCoreUObjectDelegates::OnObjectPropertyChanged.RemoveAll(this);
+
+	if (FModuleManager::Get().IsModuleLoaded(TEXT("LevelEditor")))
+	{
+		auto& LevelEditorModule = FModuleManager::LoadModuleChecked<FLevelEditorModule>(TEXT("LevelEditor"));
+		LevelEditorModule.OnMapChanged().RemoveAll(this);
+	}
 }
 
 void FFractureEditorModeToolkit::Init(const TSharedPtr<IToolkitHost>& InitToolkitHost)
@@ -200,6 +208,9 @@ void FFractureEditorModeToolkit::Init(const TSharedPtr<IToolkitHost>& InitToolki
 
 	FCoreUObjectDelegates::OnObjectPropertyChanged.AddRaw(this, &FFractureEditorModeToolkit::OnObjectPostEditChange);
 
+	FLevelEditorModule& LevelEditorModule = FModuleManager::LoadModuleChecked<FLevelEditorModule>("LevelEditor");
+	LevelEditorModule.OnMapChanged().AddRaw(this, &FFractureEditorModeToolkit::HandleMapChanged);
+
 	FDetailsViewArgs DetailsViewArgs;
 	DetailsViewArgs.bAllowSearch = false;
 	DetailsViewArgs.NameAreaSettings = FDetailsViewArgs::HideNameArea;
@@ -696,6 +707,22 @@ void FFractureEditorModeToolkit::SetHideForUnselected(UGeometryCollectionCompone
 	}
 }
 
+void FFractureEditorModeToolkit::HandleMapChanged(class UWorld* NewWorld, EMapChangeType MapChangeType)
+{
+	if ((MapChangeType == EMapChangeType::LoadMap || MapChangeType == EMapChangeType::NewMap || MapChangeType == EMapChangeType::TearDownWorld))
+	{
+		if (auto* EditorMode = GetEditorMode())
+		{ 
+			EditorMode->Exit();
+		}
+		else
+		{
+			TArray<UGeometryCollectionComponent*> EmptySelection;
+			SetOutlinerComponents(EmptySelection);
+		}
+	}
+}
+
 void FFractureEditorModeToolkit::OnToolPaletteChanged(FName PaletteName) 
 {
 	if (GetActiveTool() != nullptr)
diff --git a/Engine/Plugins/Experimental/ChaosEditor/Source/FractureEditor/Public/FractureEditorModeToolkit.h b/Engine/Plugins/Experimental/ChaosEditor/Source/FractureEditor/Public/FractureEditorModeToolkit.h
--- a/Engine/Plugins/Experimental/ChaosEditor/Source/FractureEditor/Public/FractureEditorModeToolkit.h
+++ b/Engine/Plugins/Experimental/ChaosEditor/Source/FractureEditor/Public/FractureEditorModeToolkit.h
@@ -6,6 +6,7 @@
 #include "Chaos/Real.h"
 #include "Widgets/DeclarativeSyntaxSupport.h"
 #include "Toolkits/BaseToolkit.h"
+#include "UnrealEdMisc.h"
 #include "Widgets/Input/SCheckBox.h"
 #include "Widgets/Layout/SSplitter.h"
 #include "IDetailCustomization.h"
@@ -194,6 +195,9 @@ private:
 
 	void SetHideForUnselected(UGeometryCollectionComponent* GCComp);
 
+	/** Callback for map changes. */
+	void HandleMapChanged(UWorld* NewWorld, EMapChangeType MapChangeType);
+
 private:
 	UFractureModalTool* ActiveTool;
 
