add random scatter options to mesh fracture tool

diff --git a/Engine/Plugins/Experimental/ChaosEditor/Source/FractureEditor/Private/FractureToolMeshCut.cpp b/Engine/Plugins/Experimental/ChaosEditor/Source/FractureEditor/Private/FractureToolMeshCut.cpp
--- a/Engine/Plugins/Experimental/ChaosEditor/Source/FractureEditor/Private/FractureToolMeshCut.cpp
+++ b/Engine/Plugins/Experimental/ChaosEditor/Source/FractureEditor/Private/FractureToolMeshCut.cpp
@@ -9,6 +9,7 @@
 #include "FractureToolContext.h"
 #include "Components/StaticMeshComponent.h"
 #include "Engine/StaticMesh.h"
+#include "Misc/ScopedSlowTask.h"
 
 #define LOCTEXT_NAMESPACE "FractureMesh"
 
@@ -74,6 +75,79 @@ bool UFractureToolMeshCut::IsCuttingActorValid()
 	return true;
 }
 
+void UFractureToolMeshCut::Render(const FSceneView* View, FViewport* Viewport, FPrimitiveDrawInterface* PDI)
+{
+	if (CutterSettings->bDrawDiagram && IsCuttingActorValid())
+	{
+		FBox Box = MeshCutSettings->CuttingActor->GetStaticMeshComponent()->GetStaticMesh()->GetBoundingBox();
+		
+		EnumerateVisualizationMapping(TransformsMappings, RenderMeshTransforms.Num(), [&](int32 Idx, FVector ExplodedVector)
+		{
+			const FTransform& Transform = RenderMeshTransforms[Idx];
+			FVector B000 = ExplodedVector + Transform.TransformPosition(Box.Min);
+			FVector B111 = ExplodedVector + Transform.TransformPosition(Box.Max);
+			FVector B011 = ExplodedVector + Transform.TransformPosition(FVector(Box.Min.X, Box.Max.Y, Box.Max.Z));
+			FVector B101 = ExplodedVector + Transform.TransformPosition(FVector(Box.Max.X, Box.Min.Y, Box.Max.Z));
+			FVector B110 = ExplodedVector + Transform.TransformPosition(FVector(Box.Max.X, Box.Max.Y, Box.Min.Z));
+			FVector B001 = ExplodedVector + Transform.TransformPosition(FVector(Box.Min.X, Box.Min.Y, Box.Max.Z));
+			FVector B010 = ExplodedVector + Transform.TransformPosition(FVector(Box.Min.X, Box.Max.Y, Box.Min.Z));
+			FVector B100 = ExplodedVector + Transform.TransformPosition(FVector(Box.Max.X, Box.Min.Y, Box.Min.Z));
+
+			PDI->DrawLine(B000, B100, FLinearColor::Red, SDPG_Foreground, 0.0f, 0.001f);
+			PDI->DrawLine(B000, B010, FLinearColor::Red, SDPG_Foreground, 0.0f, 0.001f);
+			PDI->DrawLine(B000, B001, FLinearColor::Red, SDPG_Foreground, 0.0f, 0.001f);
+			PDI->DrawLine(B111, B011, FLinearColor::Red, SDPG_Foreground, 0.0f, 0.001f);
+			PDI->DrawLine(B111, B101, FLinearColor::Red, SDPG_Foreground, 0.0f, 0.001f);
+			PDI->DrawLine(B111, B110, FLinearColor::Red, SDPG_Foreground, 0.0f, 0.001f);
+			PDI->DrawLine(B001, B101, FLinearColor::Red, SDPG_Foreground, 0.0f, 0.001f);
+			PDI->DrawLine(B001, B011, FLinearColor::Red, SDPG_Foreground, 0.0f, 0.001f);
+			PDI->DrawLine(B110, B100, FLinearColor::Red, SDPG_Foreground, 0.0f, 0.001f);
+			PDI->DrawLine(B110, B010, FLinearColor::Red, SDPG_Foreground, 0.0f, 0.001f);
+			PDI->DrawLine(B100, B101, FLinearColor::Red, SDPG_Foreground, 0.0f, 0.001f);
+			PDI->DrawLine(B010, B011, FLinearColor::Red, SDPG_Foreground, 0.0f, 0.001f);
+		});
+	}
+}
+
+void UFractureToolMeshCut::GenerateMeshTransforms(const FFractureToolContext& Context, TArray<FTransform>& MeshTransforms)
+{
+	FRandomStream RandStream(Context.GetSeed());
+
+	FBox Bounds = Context.GetWorldBounds();
+	const FVector Extent(Bounds.Max - Bounds.Min);
+
+	MeshTransforms.Reserve(MeshTransforms.Num() + MeshCutSettings->NumberToScatter);
+	for (int32 ii = 0; ii < MeshCutSettings->NumberToScatter; ++ii)
+	{
+		FVector Position(Bounds.Min + FVector(RandStream.FRand(), RandStream.FRand(), RandStream.FRand()) * Extent);
+		FVector::FReal Scale = RandStream.FRandRange(MeshCutSettings->MinScaleFactor, MeshCutSettings->MaxScaleFactor);
+		FVector ScaleVec(Scale, Scale, Scale);
+		MeshTransforms.Emplace(FTransform(FRotator(RandStream.FRand() * 360.0f, RandStream.FRand() * 360.0f, RandStream.FRand() * 360.0f), Position, ScaleVec));
+	}
+}
+
+void UFractureToolMeshCut::FractureContextChanged()
+{
+	UpdateDefaultRandomSeed();
+	TArray<FFractureToolContext> FractureContexts = GetFractureToolContexts();
+
+	ClearVisualizations();
+
+	for (FFractureToolContext& FractureContext : FractureContexts)
+	{
+		FBox Bounds = FractureContext.GetWorldBounds();
+		if (!Bounds.IsValid)
+		{
+			continue;
+		}
+		int32 CollectionIdx = VisualizedCollections.Add(FractureContext.GetGeometryCollectionComponent());
+		int32 BoneIdx = FractureContext.GetSelection().Num() == 1 ? FractureContext.GetSelection()[0] : INDEX_NONE;
+		TransformsMappings.AddMapping(CollectionIdx, BoneIdx, RenderMeshTransforms.Num());
+
+		GenerateMeshTransforms(FractureContext, RenderMeshTransforms);
+	}
+}
+
 int32 UFractureToolMeshCut::ExecuteFracture(const FFractureToolContext& FractureContext)
 {
 	if (FractureContext.IsValid())
@@ -86,14 +160,57 @@ int32 UFractureToolMeshCut::ExecuteFracture(const FFractureToolContext& Fracture
 		}
 
 		FMeshDescription* MeshDescription = LocalCutSettings->CuttingActor->GetStaticMeshComponent()->GetStaticMesh()->GetMeshDescription(0);
-		FTransform Transform(LocalCutSettings->CuttingActor->GetTransform());
+		FTransform ActorTransform(LocalCutSettings->CuttingActor->GetTransform());
 		FInternalSurfaceMaterials InternalSurfaceMaterials;
 
 		// Proximity is invalidated.
 		ClearProximity(FractureContext.GetGeometryCollection().Get());
 
 		// (Note: noise and grout not currently supported)
-		return CutWithMesh(MeshDescription, Transform, InternalSurfaceMaterials, *FractureContext.GetGeometryCollection(), FractureContext.GetSelection(), CollisionSettings->PointSpacing, FractureContext.GetTransform());
+		if (LocalCutSettings->NumberToScatter == 0)
+		{
+			return CutWithMesh(MeshDescription, ActorTransform, InternalSurfaceMaterials, *FractureContext.GetGeometryCollection(), FractureContext.GetSelection(), CollisionSettings->PointSpacing, FractureContext.GetTransform());
+		}
+		else
+		{
+			TArray<FTransform> MeshTransforms;
+			GenerateMeshTransforms(FractureContext, MeshTransforms);
+			int32 FirstIndex = -1;
+			// Create progress indicator dialog
+			static const FText SlowTaskText = LOCTEXT("CutWithScatteredMeshes", "Cutting geometry collection with plane(s)...");
+
+			FScopedSlowTask SlowTask(MeshTransforms.Num(), SlowTaskText);
+			SlowTask.MakeDialog();
+
+			// Declare progress shortcut lambdas
+			auto EnterProgressFrame = [&SlowTask](float Progress)
+			{
+				SlowTask.EnterProgressFrame(Progress);
+			};
+
+			// TODO: support passing multiple transforms to CutWithMesh, and doing the loop inside that function *after* the mesh conversions are done
+			// (As is, we're re-converting the cutting mesh for each transform!)
+			TArray<int32> BonesToCut = FractureContext.GetSelection();
+			for (const FTransform& ScatterTransform : MeshTransforms)
+			{
+				EnterProgressFrame(1);
+				int32 Index = CutWithMesh(MeshDescription, ScatterTransform, InternalSurfaceMaterials, *FractureContext.GetGeometryCollection(), BonesToCut, CollisionSettings->PointSpacing, FractureContext.GetTransform());
+				if (FirstIndex == -1)
+				{
+					FirstIndex = Index;
+				}
+				if (Index != -1)
+				{
+					int32 TransformIdx = FractureContext.GetGeometryCollection()->TransformIndex[Index];
+					// after a successful cut, also consider any new bones added by the cut
+					for (int32 NewBoneIdx = TransformIdx; NewBoneIdx < FractureContext.GetGeometryCollection()->NumElements(FGeometryCollection::TransformGroup); NewBoneIdx++)
+					{
+						BonesToCut.Add(NewBoneIdx);
+					}
+				}
+			}
+			return FirstIndex;
+		}
 	}
 
 	return INDEX_NONE;
diff --git a/Engine/Plugins/Experimental/ChaosEditor/Source/FractureEditor/Private/FractureToolMeshCut.h b/Engine/Plugins/Experimental/ChaosEditor/Source/FractureEditor/Private/FractureToolMeshCut.h
--- a/Engine/Plugins/Experimental/ChaosEditor/Source/FractureEditor/Private/FractureToolMeshCut.h
+++ b/Engine/Plugins/Experimental/ChaosEditor/Source/FractureEditor/Private/FractureToolMeshCut.h
@@ -23,6 +23,18 @@ public:
 	UPROPERTY(EditAnywhere, Category = MeshCut, meta = (DisplayName = "Cutting Actor"))
 	TLazyObjectPtr<AStaticMeshActor> CuttingActor;
 
+	/** Number of meshes to random scatter; if 0, the selected mesh will be used in its current position, with no random scattering*/
+	UPROPERTY(EditAnywhere, Category = Distribution)
+	int NumberToScatter = 0;
+
+	UPROPERTY(EditAnywhere, Category = Distribution, meta = (ClampMin = "0.001", EditCondition = "NumberToScatter > 0"))
+	float MinScaleFactor = .5;
+
+	UPROPERTY(EditAnywhere, Category = Distribution, meta = (ClampMin = "0.001", EditCondition = "NumberToScatter > 0"))
+	float MaxScaleFactor = 1.5;
+
+	UPROPERTY(EditAnywhere, Category = Distribution, meta = (EditCondition = "NumberToScatter > 0"))
+	bool RandomizeOrientation = true;
 };
 
 
@@ -39,12 +51,23 @@ public:
 	virtual FText GetTooltipText() const override;
 	virtual FSlateIcon GetToolIcon() const override;
 
+	void Render(const FSceneView* View, FViewport* Viewport, FPrimitiveDrawInterface* PDI) override;
+
 	virtual void RegisterUICommand(FFractureEditorCommands* BindingContext) override;
 
 	virtual TArray<UObject*> GetSettingsObjects() const override;
 
+	virtual void FractureContextChanged() override;
 	virtual int32 ExecuteFracture(const FFractureToolContext& FractureContext) override;
 
+protected:
+	virtual void ClearVisualizations() override
+	{
+		Super::ClearVisualizations();
+		RenderMeshTransforms.Empty();
+		TransformsMappings.Empty();
+	}
+
 private:
 	// Slicing
 	UPROPERTY(EditAnywhere, Category = Slicing)
@@ -52,6 +75,11 @@ private:
 
 	// check if the chosen actor can be used to cut the geometry collection (i.e. if it is a valid static mesh actor with a non-empty static mesh)
 	bool IsCuttingActorValid();
+
+	void GenerateMeshTransforms(const FFractureToolContext& Context, TArray<FTransform>& MeshTransforms);
+
+	TArray<FTransform> RenderMeshTransforms;
+	FVisualizationMappings TransformsMappings;
 };
 
 
