Add a custom pattern tool to support more user control of Voronoi fractures

diff --git a/Engine/Plugins/Experimental/ChaosEditor/Source/FractureEditor/Private/FractureEditorCommands.h b/Engine/Plugins/Experimental/ChaosEditor/Source/FractureEditor/Private/FractureEditorCommands.h
--- a/Engine/Plugins/Experimental/ChaosEditor/Source/FractureEditor/Private/FractureEditorCommands.h
+++ b/Engine/Plugins/Experimental/ChaosEditor/Source/FractureEditor/Private/FractureEditorCommands.h
@@ -59,6 +59,7 @@ class FFractureEditorCommands : public TCommands<FFractureEditorCommands>
 		TSharedPtr< FUICommandInfo > Uniform;
 		TSharedPtr< FUICommandInfo > Radial;
 		TSharedPtr< FUICommandInfo > Clustered;
+		TSharedPtr< FUICommandInfo > CustomVoronoi;
 		TSharedPtr< FUICommandInfo > Planar;
 		TSharedPtr< FUICommandInfo > Slice;
 		TSharedPtr< FUICommandInfo > Brick;
diff --git a/Engine/Plugins/Experimental/ChaosEditor/Source/FractureEditor/Private/FractureEditorModeToolkit.cpp b/Engine/Plugins/Experimental/ChaosEditor/Source/FractureEditor/Private/FractureEditorModeToolkit.cpp
--- a/Engine/Plugins/Experimental/ChaosEditor/Source/FractureEditor/Private/FractureEditorModeToolkit.cpp
+++ b/Engine/Plugins/Experimental/ChaosEditor/Source/FractureEditor/Private/FractureEditorModeToolkit.cpp
@@ -547,6 +547,7 @@ void FFractureEditorModeToolkit::BuildToolPalette(FName PaletteIndex, class FToo
 		ToolbarBuilder.AddToolBarButton(Commands.Uniform);
 		ToolbarBuilder.AddToolBarButton(Commands.Clustered);
 		ToolbarBuilder.AddToolBarButton(Commands.Radial);
+		ToolbarBuilder.AddToolBarButton(Commands.CustomVoronoi);
 		ToolbarBuilder.AddToolBarButton(Commands.Planar);
 		ToolbarBuilder.AddToolBarButton(Commands.Mesh);
 		ToolbarBuilder.AddToolBarButton(Commands.Slice);
diff --git a/Engine/Plugins/Experimental/ChaosEditor/Source/FractureEditor/Private/FractureToolCustomVoronoi.cpp b/Engine/Plugins/Experimental/ChaosEditor/Source/FractureEditor/Private/FractureToolCustomVoronoi.cpp
new file mode 100644
--- /dev/null
+++ b/Engine/Plugins/Experimental/ChaosEditor/Source/FractureEditor/Private/FractureToolCustomVoronoi.cpp
@@ -0,0 +1,292 @@
+// Copyright Epic Games, Inc. All Rights Reserved.
+
+#include "FractureToolCustomVoronoi.h"
+
+#include "FractureToolContext.h"
+#include "Components/StaticMeshComponent.h"
+#include "Engine/StaticMesh.h"
+#include "Engine/StaticMeshActor.h"
+#include "StaticMeshAttributes.h"
+#include "MeshDescription.h"
+
+#define LOCTEXT_NAMESPACE "FractureCustomVoronoi"
+
+
+void UFractureCustomVoronoiSettings::FreezeLiveSites()
+{
+	UFractureToolCustomVoronoi* CustomVoronoiTool = CastChecked<UFractureToolCustomVoronoi>(OwnerTool);
+	CustomVoronoiTool->FreezeLiveSites();
+	CustomVoronoiTool->FractureContextChanged();
+}
+
+void UFractureCustomVoronoiSettings::ClearFrozenSites()
+{
+	UFractureToolCustomVoronoi* CustomVoronoiTool = CastChecked<UFractureToolCustomVoronoi>(OwnerTool);
+	CustomVoronoiTool->ClearFrozenSites();
+	CustomVoronoiTool->FractureContextChanged();
+}
+
+void UFractureCustomVoronoiSettings::UnfreezeSites()
+{
+	UFractureToolCustomVoronoi* CustomVoronoiTool = CastChecked<UFractureToolCustomVoronoi>(OwnerTool);
+	CustomVoronoiTool->UnfreezeSites();
+	CustomVoronoiTool->FractureContextChanged();
+}
+
+void UFractureCustomVoronoiSettings::RegenerateLiveSites()
+{
+	UFractureToolCustomVoronoi* CustomVoronoiTool = CastChecked<UFractureToolCustomVoronoi>(OwnerTool);
+	CustomVoronoiTool->ClearLiveSites();
+	CustomVoronoiTool->FractureContextChanged();
+}
+
+
+void UFractureCustomVoronoiSettings::PostEditChangeChainProperty(FPropertyChangedChainEvent& PropertyChangedEvent)
+{
+	RegenerateLiveSites();
+
+	Super::PostEditChangeChainProperty(PropertyChangedEvent);
+}
+
+
+
+UFractureToolCustomVoronoi::UFractureToolCustomVoronoi(const FObjectInitializer& ObjInit) 
+	: Super(ObjInit) 
+{
+	CustomVoronoiSettings = NewObject<UFractureCustomVoronoiSettings>(GetTransientPackage(), UFractureCustomVoronoiSettings::StaticClass());
+	CustomVoronoiSettings->OwnerTool = this;
+	GizmoSettings = NewObject<UFractureTransformGizmoSettings>(GetTransientPackage(), UFractureTransformGizmoSettings::StaticClass());
+	GizmoSettings->OwnerTool = this;
+}
+
+void UFractureToolCustomVoronoi::Setup()
+{
+	GizmoSettings->Setup(this, ETransformGizmoSubElements::FullTranslateRotateScale);
+	// Stop scaling at 0 rather than going negative
+	GizmoSettings->TransformGizmo->SetDisallowNegativeScaling(true);
+	// allow non uniform scale even when the gizmo mode is set to "world"
+	GizmoSettings->TransformGizmo->SetIsNonUniformScaleAllowedFunction([]() {
+		return true;
+	});
+	// Always position the points with a gizmo
+	GizmoSettings->bShowUseGizmoOption = false;
+}
+
+
+void UFractureToolCustomVoronoi::Shutdown()
+{
+	GizmoSettings->Shutdown();
+}
+
+FText UFractureToolCustomVoronoi::GetDisplayText() const
+{ 
+	return FText(NSLOCTEXT("Fracture", "FractureToolCustomVoronoi", "Custom Voronoi Fracture")); 
+}
+
+FText UFractureToolCustomVoronoi::GetTooltipText() const 
+{ 
+	return FText(NSLOCTEXT("Fracture", "FractureToolCustomVoronoiTooltip", "Custom Voronoi Fracture creates Voronoi cells from a customizable distribution of points, which you can transform and layer.  Click the Fracture Button to commit the fracture to the geometry collection."));
+}
+
+FSlateIcon UFractureToolCustomVoronoi::GetToolIcon() const 
+{
+	return FSlateIcon("FractureEditorStyle", "FractureEditor.CustomVoronoi");
+}
+
+void UFractureToolCustomVoronoi::RegisterUICommand( FFractureEditorCommands* BindingContext ) 
+{
+	UI_COMMAND_EXT( BindingContext, UICommandInfo, "CustomVoronoi", "Custom", "Custom Voronoi Fracture", EUserInterfaceActionType::ToggleButton, FInputChord() );
+	BindingContext->CustomVoronoi = UICommandInfo;
+}
+
+TArray<UObject*> UFractureToolCustomVoronoi::GetSettingsObjects() const 
+{ 
+	TArray<UObject*> Settings;
+	Settings.Add(CustomVoronoiSettings);
+	Settings.Add(CutterSettings);
+	Settings.Add(GizmoSettings);
+	Settings.Add(CollisionSettings);
+	return Settings;
+}
+
+void UFractureToolCustomVoronoi::GenerateLivePattern(int32 RandomSeed)
+{
+	if (!ensure(CombinedWorldBounds.IsValid))
+	{
+		return;
+	}
+
+	if (!ensure(LiveSites.Num() == 0)) // We don't call this when we already have sites generated
+	{
+		LiveSites.Empty();
+	}
+
+	FRandomStream RandStream(RandomSeed);
+
+	FBox& Bounds = CombinedWorldBounds; // shorter name for convenience
+	TArray<FVector>& Sites = LiveSites;
+
+	const FVector Extent(Bounds.Max - Bounds.Min);
+	int32 NumSites = CustomVoronoiSettings->SitesToAdd;
+
+	FTransform GizmoTransform = GetGizmoTransform();
+
+	if (CustomVoronoiSettings->VoronoiPattern == EVoronoiPattern::Uniform)
+	{
+		// Uniform noise on the current bounds
+		for (int32 Idx = 0; Idx < NumSites; ++Idx)
+		{
+			Sites.Emplace(Bounds.Min + FVector(RandStream.FRand(), RandStream.FRand(), RandStream.FRand()) * Extent);
+		}
+	}
+	else if (CustomVoronoiSettings->VoronoiPattern == EVoronoiPattern::Centered)
+	{
+		// Noise centered on the gizmo
+		if (CustomVoronoiSettings->Variability > 0)
+		{
+			for (int32 Idx = 0; Idx < NumSites; ++Idx)
+			{
+				Sites.Add(GizmoTransform.GetLocation());
+			}
+		}
+		else
+		{
+			// if there's no variability, all sites would be at the same spot so we only need one of them
+			Sites.Add(GizmoTransform.GetLocation());
+		}
+	}
+	else if (CustomVoronoiSettings->VoronoiPattern == EVoronoiPattern::MeshVertices)
+	{
+		auto HasMesh = [this]() -> bool
+		{
+			if (CustomVoronoiSettings->ReferenceMesh == nullptr)
+			{
+				return false;
+			}
+			const UStaticMeshComponent* Component = CustomVoronoiSettings->ReferenceMesh->GetStaticMeshComponent();
+			if (Component == nullptr)
+			{
+				return false;
+			}
+			const UStaticMesh* Mesh = Component->GetStaticMesh();
+			if (Mesh == nullptr)
+			{
+				return false;
+			}
+			if (Mesh->GetNumLODs() < 1)
+			{
+				return false;
+			}
+			return true;
+		};
+		if (HasMesh())
+		{
+			FTransform VerticesTransform = CustomVoronoiSettings->ReferenceMesh->GetStaticMeshComponent()->GetComponentTransform();
+			if (!CustomVoronoiSettings->bStartAtActor)
+			{
+				VerticesTransform.SetLocation(GizmoTransform.GetLocation());
+			}
+
+			FMeshDescription* MeshDescription = CustomVoronoiSettings->ReferenceMesh->GetStaticMeshComponent()->GetStaticMesh()->GetMeshDescription(0);
+			TVertexAttributesConstRef<FVector3f> VertexPositions = MeshDescription->GetVertexPositions();
+			// copy vertex positions
+			for (const FVertexID VertexID : MeshDescription->Vertices().GetElementIDs())
+			{
+				const FVector Position = VertexPositions.Get(VertexID);
+				Sites.Add(VerticesTransform.TransformPosition(Position));
+			}
+		}
+	}
+	else if (CustomVoronoiSettings->VoronoiPattern == EVoronoiPattern::Grid)
+	{
+		auto ToFrac = [](int32 Val, int32 NumVals) -> FVector::FReal
+		{
+			return (FVector::FReal(Val) + FVector::FReal(.5)) / FVector::FReal(NumVals);
+		};
+		for (int32 X = 0; X < CustomVoronoiSettings->GridX; ++X)
+		{
+			FVector::FReal XFrac = ToFrac(X, CustomVoronoiSettings->GridX);
+			for (int32 Y = 0; Y < CustomVoronoiSettings->GridY; ++Y)
+			{
+				FVector::FReal YFrac = ToFrac(Y, CustomVoronoiSettings->GridY);
+				for (int32 Z = 0; Z < CustomVoronoiSettings->GridZ; ++Z)
+				{
+					FVector::FReal ZFrac = ToFrac(Z, CustomVoronoiSettings->GridZ);
+					Sites.Emplace(Bounds.Min + FVector(XFrac, YFrac, ZFrac) * Extent);
+				}
+			}
+		}
+	}
+
+	if (CustomVoronoiSettings->Variability > 0)
+	{
+		for (FVector& Site : Sites)
+		{
+			Site += (RandStream.VRand() * RandStream.FRand() * CustomVoronoiSettings->Variability);
+		}
+	}
+
+	// Convert newly generated points from world space to local (unscaled) gizmo space
+	FTransform ReferenceFrame = GetGizmoTransform();
+	ReferenceFrame.RemoveScaling();
+	for (FVector& Site : Sites)
+	{
+		Site = ReferenceFrame.InverseTransformPosition(Site);
+	}
+	
+}
+
+FTransform UFractureToolCustomVoronoi::GetGizmoTransform() const
+{
+	if (GizmoSettings->IsGizmoEnabled())
+	{
+		return GizmoSettings->GetTransform();
+	}
+	else
+	{
+		return FTransform::Identity;
+	}
+}
+
+void UFractureToolCustomVoronoi::FractureContextChanged()
+{
+	UpdateDefaultRandomSeed();
+	TArray<FFractureToolContext> FractureContexts = GetFractureToolContexts();
+
+	CombinedWorldBounds = FBox(EForceInit::ForceInit);
+	for (const FFractureToolContext& Context : FractureContexts)
+	{
+		CombinedWorldBounds += Context.GetWorldBounds();
+	}
+
+	if (CombinedWorldBounds.IsValid && LiveSites.Num() == 0 && FractureContexts.Num() > 0)
+	{
+		GenerateLivePattern(FractureContexts[0].GetSeed());
+	}
+
+	UpdateVisualizations(FractureContexts);
+}
+
+void UFractureToolCustomVoronoi::GenerateVoronoiSites(const FFractureToolContext& Context, TArray<FVector>& Sites)
+{
+	Sites.Append(FrozenSites);
+	FTransform Transform = FTransform::Identity;
+	if (GizmoSettings->IsGizmoEnabled())
+	{
+		Transform = GizmoSettings->GetTransform();
+	}
+
+	for (FVector Site : LiveSites)
+	{
+		Sites.Add(Transform.TransformPosition(Site));
+	}
+}
+
+void UFractureToolCustomVoronoi::SelectedBonesChanged()
+{
+	GizmoSettings->ResetGizmo();
+	Super::SelectedBonesChanged();
+}
+
+
+#undef LOCTEXT_NAMESPACE
diff --git a/Engine/Plugins/Experimental/ChaosEditor/Source/FractureEditor/Private/FractureToolCustomVoronoi.h b/Engine/Plugins/Experimental/ChaosEditor/Source/FractureEditor/Private/FractureToolCustomVoronoi.h
new file mode 100644
--- /dev/null
+++ b/Engine/Plugins/Experimental/ChaosEditor/Source/FractureEditor/Private/FractureToolCustomVoronoi.h
@@ -0,0 +1,169 @@
+// Copyright Epic Games, Inc. All Rights Reserved.
+
+#pragma once
+
+#include "FractureTool.h"
+#include "FractureToolCutter.h"
+
+#include "FractureToolCustomVoronoi.generated.h"
+
+class FFractureToolContext;
+
+UENUM()
+enum class EVoronoiPattern
+{
+	// Add a single site centered on the gizmo
+	Centered,
+	// Add uniform-random sites
+	Uniform,
+	// Add a regular grid of points
+	Grid,
+	// Add a point at every mesh vertex
+	MeshVertices
+};
+
+UCLASS(config = EditorPerProjectUserSettings)
+class UFractureCustomVoronoiSettings : public UFractureToolSettings
+{
+public:
+	GENERATED_BODY()
+
+	UFractureCustomVoronoiSettings(const FObjectInitializer& ObjInit)
+		: Super(ObjInit)
+		, Variability(0.0f)
+	{}
+
+	/** Method to generate next group of voronoi sites */
+	UPROPERTY(EditAnywhere, Category = LiveVoronoiSites, meta = (DisplayName = "Pattern"))
+	EVoronoiPattern VoronoiPattern = EVoronoiPattern::Uniform;
+
+	/** Randomness of sites distribution */
+	UPROPERTY(EditAnywhere, Category = LiveVoronoiSites, meta = (DisplayName = "Variability", UIMin = "0.0", ClampMin = "0.0"))
+	float Variability = 0.0f;
+
+	/** Number of Voronoi sites to add */
+	UPROPERTY(EditAnywhere, Category = LiveVoronoiSites, meta = (DisplayName = "Sites To Add", ClampMin = "1", UIMax = "1000", ClampMax = "5000",
+		EditCondition = "VoronoiPattern != EVoronoiPattern::MeshVertices && VoronoiPattern != EVoronoiPattern::Grid", EditConditionHides))
+	int32 SitesToAdd = 20;
+
+	/** Number of sites to add to grid in X */
+	UPROPERTY(EditAnywhere, Category = LiveVoronoiSites, meta = (DisplayName = "Sites in X", ClampMin = "1", UIMax = "100", ClampMax = "5000",
+		EditCondition = "VoronoiPattern == EVoronoiPattern::Grid", EditConditionHides))
+	int32 GridX = 5;
+
+	/** Number of sites to add to grid in Y */
+	UPROPERTY(EditAnywhere, Category = LiveVoronoiSites, meta = (DisplayName = "Sites in Y", ClampMin = "1", UIMax = "100", ClampMax = "5000",
+		EditCondition = "VoronoiPattern == EVoronoiPattern::Grid", EditConditionHides))
+	int32 GridY = 5;
+
+	/** Number of sites to add to grid in Z */
+	UPROPERTY(EditAnywhere, Category = LiveVoronoiSites, meta = (DisplayName = "Sites in Z", ClampMin = "1", UIMax = "100", ClampMax = "5000",
+		EditCondition = "VoronoiPattern == EVoronoiPattern::Grid", EditConditionHides))
+	int32 GridZ = 5;
+
+	//~ TODO: Have reference mesh fall back to the input fracture mesh if not provided here
+	/** Static mesh actor to be used as a reference when adding Voronoi sites */
+	UPROPERTY(EditAnywhere, Category = LiveVoronoiSites, meta = (DisplayName = "Reference Mesh", EditCondition = "VoronoiPattern == EVoronoiPattern::MeshVertices", EditConditionHides))
+	TLazyObjectPtr<AStaticMeshActor> ReferenceMesh;
+
+	/** Whether to use the reference mesh actor's transform when placing the Voronoi sites, or center them at the current gizmo location instead */
+	UPROPERTY(EditAnywhere, Category = LiveVoronoiSites, meta = (EditCondition = "VoronoiPattern == EVoronoiPattern::MeshVertices", EditConditionHides))
+	bool bStartAtActor = false;
+
+	/** Freeze the current live Voronoi Sites based on the current parameters */
+	UFUNCTION(CallInEditor, Category = CustomVoronoi)
+	void FreezeLiveSites();
+
+	/** Remove all frozen sites */
+	UFUNCTION(CallInEditor, Category = CustomVoronoi)
+	void ClearFrozenSites();
+
+	/** Unfreeze all frozen sites */
+	UFUNCTION(CallInEditor, Category = CustomVoronoi)
+	void UnfreezeSites();
+
+	/** Re-run the live Voronoi sites generator, using the current settings and selection bounds */
+	UFUNCTION(CallInEditor, Category = CustomVoronoi)
+	void RegenerateLiveSites();
+
+	virtual void PostEditChangeChainProperty(FPropertyChangedChainEvent& PropertyChangedEvent) override;
+};
+
+
+UCLASS(DisplayName="Custom Voronoi", Category="FractureTools")
+class UFractureToolCustomVoronoi : public UFractureToolVoronoiCutterBase
+{
+public:
+	GENERATED_BODY()
+
+	UFractureToolCustomVoronoi(const FObjectInitializer& ObjInit);
+
+	// UFractureTool Interface
+	virtual FText GetDisplayText() const override;
+	virtual FText GetTooltipText() const override;
+	virtual FSlateIcon GetToolIcon() const override;
+	virtual TArray<UObject*> GetSettingsObjects() const override;
+	virtual void SelectedBonesChanged() override;
+
+	virtual void RegisterUICommand( FFractureEditorCommands* BindingContext ) override;
+
+	// CustomVoronoi Voronoi Fracture Input Settings
+	UPROPERTY(EditAnywhere, Category = Uniform)
+	TObjectPtr<UFractureCustomVoronoiSettings> CustomVoronoiSettings;
+
+	UPROPERTY(EditAnywhere, Category = Uniform)
+	TObjectPtr<UFractureTransformGizmoSettings> GizmoSettings;
+
+	virtual void Setup() override;
+	virtual void Shutdown() override;
+
+	// Clear "live" voronoi sites (sites that are being edited by current settings / gizmo) -- they will repopulate on next call to GenerateVoronoiSites()
+	void ClearLiveSites()
+	{
+		LiveSites.Empty();
+	}
+
+	void FreezeLiveSites()
+	{
+		FTransform Transform = GetGizmoTransform();
+		for (const FVector& Site : LiveSites)
+		{
+			FrozenSites.Add(Transform.TransformPosition(Site));
+		}
+		LiveSites.Empty();
+	}
+
+	void ClearFrozenSites()
+	{
+		FrozenSites.Empty();
+	}
+
+	void UnfreezeSites()
+	{
+		FTransform Transform = GetGizmoTransform();
+		for (const FVector& Site : FrozenSites)
+		{
+			LiveSites.Add(Transform.InverseTransformPosition(Site));
+		}
+		FrozenSites.Empty();
+	}
+
+	// Re-calculate the "live transformed" voronoi sites based on current settings
+	void GenerateLivePattern(int32 RandomSeed);
+
+	void FractureContextChanged() override;
+
+protected:
+	void GenerateVoronoiSites(const FFractureToolContext& Context, TArray<FVector>& Sites) override;
+
+	FTransform GetGizmoTransform() const;
+
+	// Sites that are not currently live-edited
+	TArray<FVector> FrozenSites;
+
+	// Voronoi sites that are currently being regenerated and moved
+	TArray<FVector> LiveSites;
+
+	// Bounds of the combined, currently-selected fracture contexts, to be used the next time the fracture pattern is generated
+	FBox CombinedWorldBounds = FBox(EForceInit::ForceInit);
+};
\ No newline at end of file
diff --git a/Engine/Plugins/Experimental/ChaosEditor/Source/FractureEditor/Private/FractureToolCutter.cpp b/Engine/Plugins/Experimental/ChaosEditor/Source/FractureEditor/Private/FractureToolCutter.cpp
--- a/Engine/Plugins/Experimental/ChaosEditor/Source/FractureEditor/Private/FractureToolCutter.cpp
+++ b/Engine/Plugins/Experimental/ChaosEditor/Source/FractureEditor/Private/FractureToolCutter.cpp
@@ -73,7 +73,7 @@ void UFractureTransformGizmoSettings::TransformChanged(UTransformProxy* Proxy, F
 	}
 }
 
-void UFractureTransformGizmoSettings::Setup(UFractureToolCutterBase* Cutter)
+void UFractureTransformGizmoSettings::Setup(UFractureToolCutterBase* Cutter, ETransformGizmoSubElements GizmoElements)
 {
 	AttachedCutter = Cutter;
 	UInteractiveToolsContext* Context = GLevelEditorModeTools().GetInteractiveToolsContext();
@@ -81,7 +81,7 @@ void UFractureTransformGizmoSettings::Setup(UFractureToolCutterBase* Cutter)
 	{
 		UInteractiveGizmoManager* GizmoManager = Context->GizmoManager;
 		TransformProxy = NewObject<UTransformProxy>(this);
-		TransformGizmo = GizmoManager->CreateCustomTransformGizmo(ETransformGizmoSubElements::StandardTranslateRotate, this);
+		TransformGizmo = GizmoManager->CreateCustomTransformGizmo(GizmoElements, this);
 		TransformGizmo->SetActiveTarget(TransformProxy);
 		TransformProxy->OnTransformChanged.AddUObject(this, &UFractureTransformGizmoSettings::TransformChanged);
 		ResetGizmo();
@@ -247,11 +247,8 @@ void UFractureToolVoronoiCutterBase::Render(const FSceneView* View, FViewport* V
 	}
 }
 
-void UFractureToolVoronoiCutterBase::FractureContextChanged()
+void UFractureToolVoronoiCutterBase::UpdateVisualizations(TArray<FFractureToolContext>& FractureContexts)
 {
-	UpdateDefaultRandomSeed();
-	TArray<FFractureToolContext> FractureContexts = GetFractureToolContexts();
-
 	ClearVisualizations();
 
 	int32 MaxSitesToShowEdges = 100000; // computing all the voronoi diagrams can make the program non-responsive above this
@@ -299,6 +296,13 @@ void UFractureToolVoronoiCutterBase::FractureContextChanged()
 	}
 }
 
+void UFractureToolVoronoiCutterBase::FractureContextChanged()
+{
+	UpdateDefaultRandomSeed();
+	TArray<FFractureToolContext> FractureContexts = GetFractureToolContexts();
+	UpdateVisualizations(FractureContexts);
+}
+
 int32 UFractureToolVoronoiCutterBase::ExecuteFracture(const FFractureToolContext& FractureContext)
 {
 	if (FractureContext.IsValid())
diff --git a/Engine/Plugins/Experimental/ChaosEditor/Source/FractureEditor/Private/FractureToolCutter.h b/Engine/Plugins/Experimental/ChaosEditor/Source/FractureEditor/Private/FractureToolCutter.h
--- a/Engine/Plugins/Experimental/ChaosEditor/Source/FractureEditor/Private/FractureToolCutter.h
+++ b/Engine/Plugins/Experimental/ChaosEditor/Source/FractureEditor/Private/FractureToolCutter.h
@@ -194,6 +194,8 @@ protected:
 		EdgesMappings.Empty();
 	}
 
+	virtual void UpdateVisualizations(TArray<FFractureToolContext>& FractureContexts);
+
 private:
 	TArray<int32> CellMember;
 	TArray<TTuple<FVector, FVector>> VoronoiEdges;
@@ -215,7 +217,7 @@ public:
 
 	UFractureTransformGizmoSettings(const FObjectInitializer & ObjInit);
 
-	void Setup(UFractureToolCutterBase* Cutter);
+	void Setup(UFractureToolCutterBase* Cutter, ETransformGizmoSubElements GizmoElements = ETransformGizmoSubElements::StandardTranslateRotate);
 	void Shutdown();
 
 	virtual void PostEditChangeProperty(FPropertyChangedEvent& PropertyChangedEvent) override;
