fracture ui perf fix -- update outliner tree's selection in single call instead of a call per selected item, because each call required iterating the list

#rnx

diff --git a/Engine/Plugins/Experimental/ChaosEditor/Source/FractureEditor/Private/SGeometryCollectionOutliner.cpp b/Engine/Plugins/Experimental/ChaosEditor/Source/FractureEditor/Private/SGeometryCollectionOutliner.cpp
--- a/Engine/Plugins/Experimental/ChaosEditor/Source/FractureEditor/Private/SGeometryCollectionOutliner.cpp
+++ b/Engine/Plugins/Experimental/ChaosEditor/Source/FractureEditor/Private/SGeometryCollectionOutliner.cpp
@@ -291,6 +291,7 @@ void SGeometryCollectionOutliner::SetBoneSelection(UGeometryCollectionComponent*
 
 	bool bFirstSelection = true;
 
+	FGeometryCollectionTreeItemList NewSelection;
 	for(auto& RootNode : RootNodes)
 	{
 		if (RootNode->GetComponent() == RootComponent)
@@ -305,12 +306,13 @@ void SGeometryCollectionOutliner::SetBoneSelection(UGeometryCollectionComponent*
 						TreeView->RequestScrollIntoView(Item);
 						bFirstSelection = false;
 					}
-					TreeView->SetItemSelection(Item, true);
+					NewSelection.Add(Item);
 				}
 			}
 			break;
 		}
 	}
+	TreeView->SetItemSelection(NewSelection, true);
 }
 
 void SGeometryCollectionOutliner::OnSelectionChanged(FGeometryCollectionTreeItemPtr Item, ESelectInfo::Type SelectInfo)
