fix crash on autocluster when there are no bones

diff --git a/Engine/Plugins/Experimental/ChaosEditor/Source/FractureEditor/Private/FractureToolAutoCluster.cpp b/Engine/Plugins/Experimental/ChaosEditor/Source/FractureEditor/Private/FractureToolAutoCluster.cpp
--- a/Engine/Plugins/Experimental/ChaosEditor/Source/FractureEditor/Private/FractureToolAutoCluster.cpp
+++ b/Engine/Plugins/Experimental/ChaosEditor/Source/FractureEditor/Private/FractureToolAutoCluster.cpp
@@ -78,6 +78,10 @@ void UFractureToolAutoCluster::Execute(TWeakPtr<FFractureEditorModeToolkit> InTo
 			{
 				FVoronoiPartitioner VoronoiPartition(GeometryCollection, ClusterIndex);
 				VoronoiPartition.KMeansPartition(AutoClusterSettings->SiteCount);
+				if (VoronoiPartition.GetPartitionCount() == 0)
+				{
+					continue;
+				}
 
 				if (AutoClusterSettings->bEnforceConnectivity)
 				{
@@ -320,7 +324,11 @@ void FVoronoiPartitioner::InitializePartitions()
 	// At beginning, all nodes belong to first partition.
 	Partitions.Init(0, TransformIndices.Num());
 	PartitionSize.Init(0, PartitionCount);
-	PartitionSize[0] = TransformIndices.Num();
+
+	if (PartitionCount > 0)
+	{
+		PartitionSize[0] = TransformIndices.Num();
+	}
 }
 
 bool FVoronoiPartitioner::Refine()
