rename fracture's TinyGeo tool to GeoMrg, because it can also merge manually-selected bones

#rnx

diff --git a/Engine/Plugins/Experimental/ChaosEditor/Source/FractureEditor/Private/FractureToolFixTinyGeo.cpp b/Engine/Plugins/Experimental/ChaosEditor/Source/FractureEditor/Private/FractureToolFixTinyGeo.cpp
--- a/Engine/Plugins/Experimental/ChaosEditor/Source/FractureEditor/Private/FractureToolFixTinyGeo.cpp
+++ b/Engine/Plugins/Experimental/ChaosEditor/Source/FractureEditor/Private/FractureToolFixTinyGeo.cpp
@@ -9,7 +9,7 @@
 #include "FractureEditorModeToolkit.h"
 #include "FractureToolContext.h"
 
-#define LOCTEXT_NAMESPACE "FractureFixTinyGeo"
+#define LOCTEXT_NAMESPACE "FractureGeoMerge"
 
 
 UFractureToolFixTinyGeo::UFractureToolFixTinyGeo(const FObjectInitializer& ObjInit) 
@@ -21,12 +21,12 @@ UFractureToolFixTinyGeo::UFractureToolFixTinyGeo(const FObjectInitializer& ObjIn
 
 FText UFractureToolFixTinyGeo::GetDisplayText() const
 {
-	return FText(NSLOCTEXT("FixTinyGeo", "FractureToolFixTinyGeo", "Fix Too-Small Geometry")); 
+	return FText(NSLOCTEXT("FixTinyGeo", "FractureToolFixTinyGeo", "Geometry Merge Tool")); 
 }
 
 FText UFractureToolFixTinyGeo::GetTooltipText() const 
 {
-	return FText(NSLOCTEXT("FixTinyGeo", "FractureToolFixTinyGeoTooltip", "The TinyGeo tool glues \"too small\" pieces of geometry onto their neighbors."));
+	return FText(NSLOCTEXT("FixTinyGeo", "FractureToolFixTinyGeoTooltip", "The GeoMrg tool glues pieces of geometry onto their neighbors -- use it to, for example, clean up \"too small\" pieces of geometry."));
 }
 
 FSlateIcon UFractureToolFixTinyGeo::GetToolIcon() const 
@@ -36,7 +36,7 @@ FSlateIcon UFractureToolFixTinyGeo::GetToolIcon() const
 
 void UFractureToolFixTinyGeo::RegisterUICommand( FFractureEditorCommands* BindingContext ) 
 {
-	UI_COMMAND_EXT( BindingContext, UICommandInfo, "FixTinyGeo", "TinyGeo", "FixTinyGeo", EUserInterfaceActionType::ToggleButton, FInputChord() );
+	UI_COMMAND_EXT( BindingContext, UICommandInfo, "FixTinyGeo", "GeoMrg", "Geometry Merge Tool", EUserInterfaceActionType::ToggleButton, FInputChord() );
 	BindingContext->FixTinyGeo = UICommandInfo;
 }
 
diff --git a/Engine/Plugins/Experimental/ChaosEditor/Source/FractureEditor/Private/FractureToolFixTinyGeo.h b/Engine/Plugins/Experimental/ChaosEditor/Source/FractureEditor/Private/FractureToolFixTinyGeo.h
--- a/Engine/Plugins/Experimental/ChaosEditor/Source/FractureEditor/Private/FractureToolFixTinyGeo.h
+++ b/Engine/Plugins/Experimental/ChaosEditor/Source/FractureEditor/Private/FractureToolFixTinyGeo.h
@@ -27,7 +27,7 @@ enum class ENeighborSelectionMethod
 	NearestCenter
 };
 
-/** Settings controlling how tiny geometry is selected and merged into neighboring geometry */
+/** Settings controlling how geometry is selected and merged into neighboring geometry */
 UCLASS(config = EditorPerProjectUserSettings)
 class UFractureTinyGeoSettings : public UFractureToolSettings
 {
@@ -39,11 +39,16 @@ public:
 		: Super(ObjInit)
 	{}
 
-	UPROPERTY(EditAnywhere, Category = FilterSettings)
-	EGeometrySelectionMethod SelectionMethod = EGeometrySelectionMethod::RelativeVolume;
+	UPROPERTY(EditAnywhere, Category = MergeSettings)
+	ENeighborSelectionMethod NeighborSelection = ENeighborSelectionMethod::LargestNeighbor;
+
+	/** Also merge selected bones to their neighbors */
+	UPROPERTY(EditAnywhere, Category = MergeSettings, meta = (DisplayName = "Merge Selected"))
+	bool bAlsoMergeSelected = false;
+
 
 	UPROPERTY(EditAnywhere, Category = FilterSettings)
-	ENeighborSelectionMethod NeighborSelection = ENeighborSelectionMethod::LargestNeighbor;
+	EGeometrySelectionMethod SelectionMethod = EGeometrySelectionMethod::RelativeVolume;
 
 	/** If volume is less than this value cubed, geometry should be merged into neighbors -- i.e. a value of 2 merges geometry smaller than a 2x2x2 cube */
 	UPROPERTY(EditAnywhere, Category = FilterSettings, meta = (ClampMin = ".00001", UIMin = ".1", UIMax = "10", EditCondition = "SelectionMethod == EGeometrySelectionMethod::VolumeCubeRoot", EditConditionHides))
@@ -54,16 +59,12 @@ public:
 	UPROPERTY(EditAnywhere, Category = FilterSettings, meta = (ClampMin = "0", UIMax = ".1", ClampMax = "1.0", EditCondition = "SelectionMethod == EGeometrySelectionMethod::RelativeVolume", EditConditionHides))
 	double RelativeVolume = .01;
 
-	/** Also merge selected bones to their neighbors */
-	UPROPERTY(EditAnywhere, Category = FilterSettings)
-	bool bAlsoMergeSelected = false;
-
 };
 
 
 
 // Note this tool doesn't actually fracture, but it does remake pieces of geometry and shares a lot of machinery with the fracture tools
-UCLASS(DisplayName = "Fix Tiny Geometry Tool", Category = "FractureTools")
+UCLASS(DisplayName = "Geometry Merge Tool", Category = "FractureTools")
 class UFractureToolFixTinyGeo : public UFractureToolCutterBase
 {
 public:
@@ -76,7 +77,7 @@ public:
 	virtual FText GetTooltipText() const override;
 	virtual FSlateIcon GetToolIcon() const override;
 
-	virtual FText GetApplyText() const override { return FText(NSLOCTEXT("FixTinyGeo", "ExecuteFixTinyGeo", "Merge Geometry")); }
+	virtual FText GetApplyText() const override { return FText(NSLOCTEXT("FractureGeoMerge", "ExecuteGeoMerge", "Merge Geometry")); }
 
 	void Render(const FSceneView* View, FViewport* Viewport, FPrimitiveDrawInterface* PDI) override;
 
