add a 'skip fraction' to the custom voronoi tool, to remove a percentage of the 'live' sites

diff --git a/Engine/Plugins/Experimental/ChaosEditor/Source/FractureEditor/Private/FractureToolCustomVoronoi.cpp b/Engine/Plugins/Experimental/ChaosEditor/Source/FractureEditor/Private/FractureToolCustomVoronoi.cpp
--- a/Engine/Plugins/Experimental/ChaosEditor/Source/FractureEditor/Private/FractureToolCustomVoronoi.cpp
+++ b/Engine/Plugins/Experimental/ChaosEditor/Source/FractureEditor/Private/FractureToolCustomVoronoi.cpp
@@ -226,6 +226,14 @@ void UFractureToolCustomVoronoi::GenerateLivePattern(int32 RandomSeed)
 		}
 	}
 
+	// randomly remove points based on the skip fraction
+	int32 TargetNumSites = Sites.Num() - int32(Sites.Num() * CustomVoronoiSettings->SkipFraction);
+	while (TargetNumSites < Sites.Num())
+	{
+		int32 ToRemoveIdx = RandStream.RandHelper(Sites.Num());
+		Sites.RemoveAtSwap(ToRemoveIdx, 1, false);
+	}
+
 	// Convert newly generated points from world space to local (unscaled) gizmo space
 	FTransform ReferenceFrame = GetGizmoTransform();
 	ReferenceFrame.RemoveScaling();
diff --git a/Engine/Plugins/Experimental/ChaosEditor/Source/FractureEditor/Private/FractureToolCustomVoronoi.h b/Engine/Plugins/Experimental/ChaosEditor/Source/FractureEditor/Private/FractureToolCustomVoronoi.h
--- a/Engine/Plugins/Experimental/ChaosEditor/Source/FractureEditor/Private/FractureToolCustomVoronoi.h
+++ b/Engine/Plugins/Experimental/ChaosEditor/Source/FractureEditor/Private/FractureToolCustomVoronoi.h
@@ -61,6 +61,10 @@ public:
 		EditCondition = "VoronoiPattern == EVoronoiPattern::Grid", EditConditionHides))
 	int32 GridZ = 5;
 
+	/** Fraction of points to randomly skip */
+	UPROPERTY(EditAnywhere, Category = LiveVoronoiSites, meta = (ClampMin = "0", ClampMax = "1"))
+	float SkipFraction = 0.0;
+
 	//~ TODO: Have reference mesh fall back to the input fracture mesh if not provided here
 	/** Static mesh actor to be used as a reference when adding Voronoi sites */
 	UPROPERTY(EditAnywhere, Category = LiveVoronoiSites, meta = (DisplayName = "Reference Mesh", EditCondition = "VoronoiPattern == EVoronoiPattern::MeshVertices", EditConditionHides))
