fix mesh cut 'random orientation' toggle not functioning

#rnx

diff --git a/Engine/Plugins/Experimental/ChaosEditor/Source/FractureEditor/Private/FractureToolMeshCut.cpp b/Engine/Plugins/Experimental/ChaosEditor/Source/FractureEditor/Private/FractureToolMeshCut.cpp
--- a/Engine/Plugins/Experimental/ChaosEditor/Source/FractureEditor/Private/FractureToolMeshCut.cpp
+++ b/Engine/Plugins/Experimental/ChaosEditor/Source/FractureEditor/Private/FractureToolMeshCut.cpp
@@ -122,7 +122,16 @@ void UFractureToolMeshCut::GenerateMeshTransforms(const FFractureToolContext& Co
 		FVector Position(Bounds.Min + FVector(RandStream.FRand(), RandStream.FRand(), RandStream.FRand()) * Extent);
 		FVector::FReal Scale = RandStream.FRandRange(MeshCutSettings->MinScaleFactor, MeshCutSettings->MaxScaleFactor);
 		FVector ScaleVec(Scale, Scale, Scale);
-		MeshTransforms.Emplace(FTransform(FRotator(RandStream.FRand() * 360.0f, RandStream.FRand() * 360.0f, RandStream.FRand() * 360.0f), Position, ScaleVec));
+		FRotator Orientation = FRotator::ZeroRotator;
+		if (MeshCutSettings->bRandomOrientation)
+		{
+			Orientation = FRotator(
+				RandStream.FRandRange(-MeshCutSettings->PitchRange, MeshCutSettings->PitchRange),
+				RandStream.FRandRange(-MeshCutSettings->YawRange, MeshCutSettings->YawRange),
+				RandStream.FRandRange(-MeshCutSettings->RollRange, MeshCutSettings->RollRange)
+			);
+		}
+		MeshTransforms.Emplace(FTransform(Orientation, Position, ScaleVec));
 	}
 }
 
diff --git a/Engine/Plugins/Experimental/ChaosEditor/Source/FractureEditor/Private/FractureToolMeshCut.h b/Engine/Plugins/Experimental/ChaosEditor/Source/FractureEditor/Private/FractureToolMeshCut.h
--- a/Engine/Plugins/Experimental/ChaosEditor/Source/FractureEditor/Private/FractureToolMeshCut.h
+++ b/Engine/Plugins/Experimental/ChaosEditor/Source/FractureEditor/Private/FractureToolMeshCut.h
@@ -34,7 +34,19 @@ public:
 	float MaxScaleFactor = 1.5;
 
 	UPROPERTY(EditAnywhere, Category = Distribution, meta = (EditCondition = "NumberToScatter > 0"))
-	bool RandomizeOrientation = true;
+	bool bRandomOrientation = true;
+
+	/** Roll will be chosen between -Range to +Range */
+	UPROPERTY(EditAnywhere, Category = Distribution, meta = (DisplayName = "+/- Roll Range", EditCondition = "NumberToScatter > 0 && bRandomOrientation", ClampMin = "0", ClampMax = "180"))
+	float RollRange = 180;
+
+	/** Pitch will be chosen between -Range to +Range */
+	UPROPERTY(EditAnywhere, Category = Distribution, meta = (DisplayName = "+/- Pitch Range", EditCondition = "NumberToScatter > 0 && bRandomOrientation", ClampMin = "0", ClampMax = "180"))
+	float PitchRange = 180;
+
+	/** Yaw will be chosen between -Range to +Range */
+	UPROPERTY(EditAnywhere, Category = Distribution, meta = (DisplayName = "+/- Yaw Range", EditCondition = "NumberToScatter > 0 && bRandomOrientation", ClampMin = "0", ClampMax = "180"))
+	float YawRange = 180;
 };
 
 
