more correct fix for fracture crash on map change

#rnx

diff --git a/Engine/Plugins/Experimental/ChaosEditor/Source/FractureEditor/Private/FractureEditorModeToolkit.cpp b/Engine/Plugins/Experimental/ChaosEditor/Source/FractureEditor/Private/FractureEditorModeToolkit.cpp
--- a/Engine/Plugins/Experimental/ChaosEditor/Source/FractureEditor/Private/FractureEditorModeToolkit.cpp
+++ b/Engine/Plugins/Experimental/ChaosEditor/Source/FractureEditor/Private/FractureEditorModeToolkit.cpp
@@ -193,11 +193,6 @@ FFractureEditorModeToolkit::~FFractureEditorModeToolkit()
 	}
 }
 
-void FFractureEditorModeToolkit::Init(const TSharedPtr<IToolkitHost>& InitToolkitHost)
-{
-	Init(InitToolkitHost, TWeakObjectPtr<UEdMode>());
-}
-
 void FFractureEditorModeToolkit::Init(const TSharedPtr<IToolkitHost>& InitToolkitHost, TWeakObjectPtr<UEdMode> InOwningMode)
 {
 	FFractureEditorModule& FractureModule = FModuleManager::GetModuleChecked<FFractureEditorModule>("FractureEditor");
@@ -712,7 +707,7 @@ void FFractureEditorModeToolkit::HandleMapChanged(class UWorld* NewWorld, EMapCh
 {
 	if ((MapChangeType == EMapChangeType::LoadMap || MapChangeType == EMapChangeType::NewMap || MapChangeType == EMapChangeType::TearDownWorld))
 	{
-		if (auto* EditorMode = GetEditorMode())
+		if (UEdMode* EditorMode = OwningEditorMode.Get())
 		{ 
 			EditorMode->Exit();
 		}
diff --git a/Engine/Plugins/Experimental/ChaosEditor/Source/FractureEditor/Public/FractureEditorModeToolkit.h b/Engine/Plugins/Experimental/ChaosEditor/Source/FractureEditor/Public/FractureEditorModeToolkit.h
--- a/Engine/Plugins/Experimental/ChaosEditor/Source/FractureEditor/Public/FractureEditorModeToolkit.h
+++ b/Engine/Plugins/Experimental/ChaosEditor/Source/FractureEditor/Public/FractureEditorModeToolkit.h
@@ -82,7 +82,6 @@ public:
 	~FFractureEditorModeToolkit();
 	
 	/** FModeToolkit interface */
-	virtual void Init(const TSharedPtr<IToolkitHost>& InitToolkitHost) override;
 	virtual void Init(const TSharedPtr<IToolkitHost>& InitToolkitHost, TWeakObjectPtr<UEdMode> InOwningMode) override;
 	/** IToolkit interface */
 	virtual FName GetToolkitFName() const override;
