fix missing Super:: calls for setup and shutdown of tools with gizmos (custom voronoi, radial, plane cut)

#rnx

diff --git a/Engine/Plugins/Experimental/ChaosEditor/Source/FractureEditor/Private/FractureToolCustomVoronoi.cpp b/Engine/Plugins/Experimental/ChaosEditor/Source/FractureEditor/Private/FractureToolCustomVoronoi.cpp
--- a/Engine/Plugins/Experimental/ChaosEditor/Source/FractureEditor/Private/FractureToolCustomVoronoi.cpp
+++ b/Engine/Plugins/Experimental/ChaosEditor/Source/FractureEditor/Private/FractureToolCustomVoronoi.cpp
@@ -61,6 +61,7 @@ UFractureToolCustomVoronoi::UFractureToolCustomVoronoi(const FObjectInitializer&
 
 void UFractureToolCustomVoronoi::Setup()
 {
+	Super::Setup();
 	GizmoSettings->Setup(this, ETransformGizmoSubElements::FullTranslateRotateScale);
 	// Stop scaling at 0 rather than going negative
 	GizmoSettings->TransformGizmo->SetDisallowNegativeScaling(true);
@@ -75,6 +76,7 @@ void UFractureToolCustomVoronoi::Setup()
 
 void UFractureToolCustomVoronoi::Shutdown()
 {
+	Super::Shutdown();
 	GizmoSettings->Shutdown();
 }
 
diff --git a/Engine/Plugins/Experimental/ChaosEditor/Source/FractureEditor/Private/FractureToolPlaneCut.cpp b/Engine/Plugins/Experimental/ChaosEditor/Source/FractureEditor/Private/FractureToolPlaneCut.cpp
--- a/Engine/Plugins/Experimental/ChaosEditor/Source/FractureEditor/Private/FractureToolPlaneCut.cpp
+++ b/Engine/Plugins/Experimental/ChaosEditor/Source/FractureEditor/Private/FractureToolPlaneCut.cpp
@@ -24,6 +24,7 @@ UFractureToolPlaneCut::UFractureToolPlaneCut(const FObjectInitializer& ObjInit)
 
 void UFractureToolPlaneCut::Setup()
 {
+	Super::Setup();
 	GizmoSettings->Setup(this);
 	PlaneCutSettings->bCanCutWithMultiplePlanes = !GizmoSettings->bUseGizmo;
 	NotifyOfPropertyChangeByTool(PlaneCutSettings);
@@ -32,6 +33,7 @@ void UFractureToolPlaneCut::Setup()
 
 void UFractureToolPlaneCut::Shutdown()
 {
+	Super::Shutdown();
 	GizmoSettings->Shutdown();
 }
 
diff --git a/Engine/Plugins/Experimental/ChaosEditor/Source/FractureEditor/Private/FractureToolRadial.cpp b/Engine/Plugins/Experimental/ChaosEditor/Source/FractureEditor/Private/FractureToolRadial.cpp
--- a/Engine/Plugins/Experimental/ChaosEditor/Source/FractureEditor/Private/FractureToolRadial.cpp
+++ b/Engine/Plugins/Experimental/ChaosEditor/Source/FractureEditor/Private/FractureToolRadial.cpp
@@ -19,12 +19,14 @@ UFractureToolRadial::UFractureToolRadial(const FObjectInitializer& ObjInit)
 
 void UFractureToolRadial::Setup()
 {
+	Super::Setup();
 	GizmoSettings->Setup(this);
 }
 
 
 void UFractureToolRadial::Shutdown()
 {
+	Super::Shutdown();
 	GizmoSettings->Shutdown();
 }
 
