Fix the transform gizmo being completely obscured by the fracture voronoi diagram and fracture selection material

 - Change gizmo created by fracture to use editor compositing (makes it dithered unfortunately, but makes it unaffected by translucent materials)
 - Use less saturated colors for voronoi diagrams and sites (Gizmo uses full saturation red/green/blue)
 - Change voronoi diagram rendering to use line set component (makes it *not* dithered, so doesn't fight with now-dithered gizmo ... plus it's a lot faster to render)

To make the line set component easier to see under motion, also temporarily override motion blur and temporal AA when line sets are active.

#rnx

diff --git a/Engine/Plugins/Experimental/ChaosEditor/Source/FractureEditor/Private/FractureToolCutter.cpp b/Engine/Plugins/Experimental/ChaosEditor/Source/FractureEditor/Private/FractureToolCutter.cpp
--- a/Engine/Plugins/Experimental/ChaosEditor/Source/FractureEditor/Private/FractureToolCutter.cpp
+++ b/Engine/Plugins/Experimental/ChaosEditor/Source/FractureEditor/Private/FractureToolCutter.cpp
@@ -6,6 +6,8 @@
 #include "FractureToolContext.h"
 #include "InteractiveToolsContext.h"
 #include "EditorModeManager.h"
+#include "ToolSetupUtil.h"
+#include "BaseGizmos/GizmoBaseComponent.h"
 
 #include "GeometryCollection/GeometryCollectionObject.h"
 #include "GeometryCollection/GeometryCollectionComponent.h"
@@ -82,6 +84,17 @@ void UFractureTransformGizmoSettings::Setup(UFractureToolCutterBase* Cutter, ETr
 		UInteractiveGizmoManager* GizmoManager = Context->GizmoManager;
 		TransformProxy = NewObject<UTransformProxy>(this);
 		TransformGizmo = GizmoManager->CreateCustomTransformGizmo(GizmoElements, this);
+		// TODO: Stop setting bUseEditorCompositing on gizmo components like this, once gizmo rendering is improved
+		// This is a hack to make gizmos more visible when the translucent fracture bone selection material is on screen
+		// This hack unfortunately makes the gizmos dithered when they're behind objects, but at least they're relatively visible
+		for (UActorComponent* Component : TransformGizmo->GetGizmoActor()->GetComponents())
+		{
+			UGizmoBaseComponent* GizmoComponent = Cast<UGizmoBaseComponent>(Component);
+			if (GizmoComponent)
+			{
+				GizmoComponent->bUseEditorCompositing = true;
+			}
+		}
 		TransformGizmo->SetActiveTarget(TransformProxy);
 		TransformProxy->OnTransformChanged.AddUObject(this, &UFractureTransformGizmoSettings::TransformChanged);
 		ResetGizmo();
@@ -222,7 +235,10 @@ UFractureToolVoronoiCutterBase::UFractureToolVoronoiCutterBase(const FObjectInit
 {
 	for (int32 ii = 0; ii < 100; ++ii)
 	{
-		Colors.Emplace(FMath::FRand(), FMath::FRand(), FMath::FRand());;
+		// This is adapted from FColor::MakeRandomColor() but desaturated a bit
+		const uint8 Hue = (uint8)(FMath::FRand() * 255.f);
+		FColor DesatRandomColor = FLinearColor::MakeFromHSV8(Hue, 190, 255).ToFColor(true);
+		Colors.Add(DesatRandomColor);
 	}
 }
 
@@ -232,19 +248,21 @@ void UFractureToolVoronoiCutterBase::Render(const FSceneView* View, FViewport* V
 	{
 		EnumerateVisualizationMapping(SitesMappings, VoronoiSites.Num(), [&](int32 Idx, FVector ExplodedVector)
 		{
-			PDI->DrawPoint(VoronoiSites[Idx] + ExplodedVector, FLinearColor::Green, 4.f, SDPG_Foreground);
+			PDI->DrawPoint(VoronoiSites[Idx] + ExplodedVector, FLinearColor(.62, .94, .84), 4.f, SDPG_Foreground);
 		});
 	}
 
-	if (CutterSettings->bDrawDiagram)
+	UpdateLineSetExplodedVectors();
+}
+
+void UFractureToolVoronoiCutterBase::UpdateLineSetExplodedVectors()
+{
+	EnumerateVisualizationMapping(EdgesMappings, VoronoiLineSets.Num(), [&](int32 Idx, FVector ExplodedVector)
 	{
-		PDI->AddReserveLines(SDPG_Foreground, VoronoiEdges.Num(), false, false);
-		EnumerateVisualizationMapping(EdgesMappings, VoronoiEdges.Num(), [&](int32 Idx, FVector ExplodedVector)
-		{
-			PDI->DrawLine(VoronoiEdges[Idx].Get<0>() + ExplodedVector,
-				VoronoiEdges[Idx].Get<1>() + ExplodedVector, Colors[CellMember[Idx] % 100], SDPG_Foreground);
-		});
-	}
+		// TODO: If we add diagrams even when visibility is toggled off, we also need to update visibility here, e.g.: VoronoiLineSets[Idx]->SetVisibility(CutterSettings->bDrawDiagram);
+		AActor* Actor = VoronoiLineSets[Idx]->GetOwner();
+		VoronoiLineSets[Idx]->SetRelativeLocation(Actor->GetTransform().InverseTransformVector(ExplodedVector));
+	});
 }
 
 void UFractureToolVoronoiCutterBase::UpdateVisualizations(TArray<FFractureToolContext>& FractureContexts)
@@ -264,7 +282,7 @@ void UFractureToolVoronoiCutterBase::UpdateVisualizations(TArray<FFractureToolCo
 		int32 CollectionIdx = VisualizedCollections.Emplace(FractureContext.GetGeometryCollectionComponent());
 		int32 BoneIdx = FractureContext.GetSelection().Num() == 1 ? FractureContext.GetSelection()[0] : INDEX_NONE;
 		SitesMappings.AddMapping(CollectionIdx, BoneIdx, VoronoiSites.Num());
-		EdgesMappings.AddMapping(CollectionIdx, BoneIdx, VoronoiEdges.Num());
+		EdgesMappings.AddMapping(CollectionIdx, BoneIdx, VoronoiLineSets.Num());
 
 		// Generate voronoi diagrams and cache visualization info
 		TArray<FVector> LocalVoronoiSites;
@@ -280,20 +298,39 @@ void UFractureToolVoronoiCutterBase::UpdateVisualizations(TArray<FFractureToolCo
 		if (bEstAboveMaxSites || LocalVoronoiSites.Num() * FractureContexts.Num() > MaxSitesToShowEdges || VoronoiSites.Num() > MaxSitesToShowEdges)
 		{
 			UE_LOG(LogFractureTool, Warning, TEXT("Voronoi diagram(s) number of sites too large; will not display Voronoi diagram edges"));
-			VoronoiEdges.Empty();
-			CellMember.Empty();
-			EdgesMappings.Empty();
+			ClearEdges();
 			bEstAboveMaxSites = true;
 		}
-		else
+		else if (CutterSettings->bDrawDiagram)
 		{
 			FBox VoronoiBounds = GetVoronoiBounds(FractureContext, LocalVoronoiSites);
-			if (CutterSettings->bDrawDiagram)
+			ULineSetComponent* VoronoiLineSet = NewObject<ULineSetComponent>(FractureContext.GetGeometryCollectionComponent());
+			VoronoiLineSet->SetupAttachment(FractureContext.GetGeometryCollectionComponent());
+			VoronoiLineSet->SetLineMaterial(ToolSetupUtil::GetDefaultLineComponentMaterial(nullptr, false));
+			// TODO: If we switch to adding diagrams even when visibility is toggled off, we also need to update visibility here, e.g.: VoronoiLineSets[Idx]->SetVisibility(CutterSettings->bDrawDiagram);
+			VoronoiLineSet->RegisterComponent();
+			TArray<TTuple<FVector, FVector>> VoronoiEdges;
+			GetVoronoiEdges(LocalVoronoiSites, VoronoiBounds, VoronoiEdges, CellMember);
+			FTransform ToWorld = FractureContext.GetTransform();
+			VoronoiLineSet->ReserveLines(VoronoiEdges.Num());
+			for (int32 Idx = 0; Idx < VoronoiEdges.Num(); ++Idx)
 			{
-				GetVoronoiEdges(LocalVoronoiSites, VoronoiBounds, VoronoiEdges, CellMember);
+				const TTuple<FVector, FVector>& Line = VoronoiEdges[Idx];
+				VoronoiLineSet->AddLine(ToWorld.InverseTransformPosition(VoronoiEdges[Idx].Key), ToWorld.InverseTransformPosition(VoronoiEdges[Idx].Value), Colors[CellMember[Idx] % 100], 1.3, .1);
 			}
+			VoronoiLineSets.Add(VoronoiLineSet);
 		}
 	}
+
+	if (VoronoiLineSets.Num() > 0)
+	{
+		OverrideEditorViewFlagsForLineRendering();
+		UpdateLineSetExplodedVectors();
+	}
+	else
+	{
+		RestoreEditorViewFlags();
+	}
 }
 
 void UFractureToolVoronoiCutterBase::FractureContextChanged()
diff --git a/Engine/Plugins/Experimental/ChaosEditor/Source/FractureEditor/Private/FractureToolCutter.h b/Engine/Plugins/Experimental/ChaosEditor/Source/FractureEditor/Private/FractureToolCutter.h
--- a/Engine/Plugins/Experimental/ChaosEditor/Source/FractureEditor/Private/FractureToolCutter.h
+++ b/Engine/Plugins/Experimental/ChaosEditor/Source/FractureEditor/Private/FractureToolCutter.h
@@ -6,6 +6,7 @@
 
 #include "BaseGizmos/TransformGizmo.h"
 #include "BaseGizmos/TransformProxy.h"
+#include "Drawing/LineSetComponent.h"
 
 #include "FractureToolCutter.generated.h"
 
@@ -187,22 +188,36 @@ protected:
 	virtual void ClearVisualizations()
 	{
 		Super::ClearVisualizations();
-		CellMember.Empty();
-		VoronoiEdges.Empty();
 		VoronoiSites.Empty();
 		SitesMappings.Empty();
+		ClearEdges();
+	}
+
+	void ClearEdges()
+	{
+		CellMember.Empty();
 		EdgesMappings.Empty();
+		for (ULineSetComponent* Lines : VoronoiLineSets)
+		{
+			Lines->DestroyComponent();
+		}
+		VoronoiLineSets.Empty();
 	}
 
 	virtual void UpdateVisualizations(TArray<FFractureToolContext>& FractureContexts);
 
 private:
+	UPROPERTY()
+	TArray<TObjectPtr<ULineSetComponent>> VoronoiLineSets;
+
 	TArray<int32> CellMember;
-	TArray<TTuple<FVector, FVector>> VoronoiEdges;
 	TArray<FVector> VoronoiSites;
-	TArray<FLinearColor> Colors;
+
+	TArray<FColor> Colors;
 
 	FVisualizationMappings SitesMappings, EdgesMappings;
+
+	void UpdateLineSetExplodedVectors();
 };
 
 
diff --git a/Engine/Plugins/Experimental/ChaosEditor/Source/FractureEditor/Public/FractureTool.cpp b/Engine/Plugins/Experimental/ChaosEditor/Source/FractureEditor/Public/FractureTool.cpp
--- a/Engine/Plugins/Experimental/ChaosEditor/Source/FractureEditor/Public/FractureTool.cpp
+++ b/Engine/Plugins/Experimental/ChaosEditor/Source/FractureEditor/Public/FractureTool.cpp
@@ -4,6 +4,9 @@
 
 #include "Editor.h"
 #include "Engine/Selection.h"
+#include "LevelEditor.h"
+#include "LevelEditorViewport.h"
+#include "SLevelViewport.h"
 
 #include "FractureToolContext.h"
 #include "FractureSelectionTools.h"
@@ -210,6 +213,46 @@ void UFractureModalTool::EnumerateVisualizationMapping(const FVisualizationMappi
 }
 
 
+void UFractureModalTool::OverrideEditorViewFlagsForLineRendering()
+{
+	FLevelEditorModule& LevelEditorModule = FModuleManager::GetModuleChecked<FLevelEditorModule>("LevelEditor");
+	TSharedPtr<ILevelEditor> LevelEditor = LevelEditorModule.GetFirstLevelEditor();
+	if (LevelEditor.IsValid())
+	{
+		TArray<TSharedPtr<SLevelViewport>> Viewports = LevelEditor->GetViewports();
+		for (const TSharedPtr<SLevelViewport>& ViewportWindow : Viewports)
+		{
+			if (ViewportWindow.IsValid())
+			{
+				FEditorViewportClient& Viewport = ViewportWindow->GetAssetViewportClient();
+				Viewport.EnableOverrideEngineShowFlags([](FEngineShowFlags& Flags)
+					{
+						Flags.SetTemporalAA(false);
+						Flags.SetMotionBlur(false);
+					});
+			}
+		}
+	}
+}
+
+void UFractureModalTool::RestoreEditorViewFlags()
+{
+	FLevelEditorModule& LevelEditorModule = FModuleManager::GetModuleChecked<FLevelEditorModule>("LevelEditor");
+	TSharedPtr<ILevelEditor> LevelEditor = LevelEditorModule.GetFirstLevelEditor();
+	if (LevelEditor.IsValid())
+	{
+		TArray<TSharedPtr<SLevelViewport>> Viewports = LevelEditor->GetViewports();
+		for (const TSharedPtr<SLevelViewport>& ViewportWindow : Viewports)
+		{
+			if (ViewportWindow.IsValid())
+			{
+				FEditorViewportClient& Viewport = ViewportWindow->GetAssetViewportClient();
+				Viewport.DisableOverrideEngineShowFlags();
+			}
+		}
+	}
+}
+
 void UFractureModalTool::NotifyOfPropertyChangeByTool(UFractureToolSettings* PropertySet) const
 {
 	OnPropertyModifiedDirectlyByTool.Broadcast(PropertySet);
diff --git a/Engine/Plugins/Experimental/ChaosEditor/Source/FractureEditor/Public/FractureTool.h b/Engine/Plugins/Experimental/ChaosEditor/Source/FractureEditor/Public/FractureTool.h
--- a/Engine/Plugins/Experimental/ChaosEditor/Source/FractureEditor/Public/FractureTool.h
+++ b/Engine/Plugins/Experimental/ChaosEditor/Source/FractureEditor/Public/FractureTool.h
@@ -159,6 +159,7 @@ public:
 	{
 		GEngine->OnComponentTransformChanged().RemoveAll(this);
 		ClearVisualizations();
+		RestoreEditorViewFlags();
 	}
 
 	// Called when a selected geometry collection component is moved in the scene
@@ -198,6 +199,11 @@ protected:
 	
 	void OnComponentTransformChangedInternal(USceneComponent* InRootComponent, ETeleportType Teleport);
 
+	// Call to override editor view flags to disable temporal AA and motion blur, which make the lines look bad
+	void OverrideEditorViewFlagsForLineRendering();
+	// Restore editor view flags
+	void RestoreEditorViewFlags();
+
 };
 
 
