fix crash on map change when fracture tool was open (because UEdMode::Exit was being called twice)

#rnx

diff --git a/Engine/Plugins/Experimental/ChaosEditor/Source/FractureEditor/Private/FractureEditorModeToolkit.cpp b/Engine/Plugins/Experimental/ChaosEditor/Source/FractureEditor/Private/FractureEditorModeToolkit.cpp
--- a/Engine/Plugins/Experimental/ChaosEditor/Source/FractureEditor/Private/FractureEditorModeToolkit.cpp
+++ b/Engine/Plugins/Experimental/ChaosEditor/Source/FractureEditor/Private/FractureEditorModeToolkit.cpp
@@ -707,15 +707,9 @@ void FFractureEditorModeToolkit::HandleMapChanged(class UWorld* NewWorld, EMapCh
 {
 	if ((MapChangeType == EMapChangeType::LoadMap || MapChangeType == EMapChangeType::NewMap || MapChangeType == EMapChangeType::TearDownWorld))
 	{
-		if (UEdMode* EditorMode = OwningEditorMode.Get())
-		{ 
-			EditorMode->Exit();
-		}
-		else
-		{
-			TArray<UGeometryCollectionComponent*> EmptySelection;
-			SetOutlinerComponents(EmptySelection);
-		}
+		ShutdownActiveTool();
+		TArray<UGeometryCollectionComponent*> EmptySelection;
+		SetOutlinerComponents(EmptySelection);
 	}
 }
 
