Adopt UI customization features from Modeling Tools Mode into Fracture Tools
Also update the UI customization to accept localized display names as well as original/non-localized names

#rnx

diff --git a/Engine/Plugins/Editor/ModelingToolsEditorMode/Source/ModelingToolsEditorMode/Private/ModelingToolsEditorModeToolkit.cpp b/Engine/Plugins/Editor/ModelingToolsEditorMode/Source/ModelingToolsEditorMode/Private/ModelingToolsEditorModeToolkit.cpp
--- a/Engine/Plugins/Editor/ModelingToolsEditorMode/Source/ModelingToolsEditorMode/Private/ModelingToolsEditorModeToolkit.cpp
+++ b/Engine/Plugins/Editor/ModelingToolsEditorMode/Source/ModelingToolsEditorMode/Private/ModelingToolsEditorModeToolkit.cpp
@@ -516,7 +516,8 @@ void FModelingToolsEditorModeToolkit::GetToolPaletteNames(TArray<FName>& Palette
 		{
 			for (int32 k = 0; k < PaletteNames.Num(); ++k)
 			{
-				if (SectionName.Compare(PaletteNames[k].ToString(), ESearchCase::IgnoreCase) == 0)
+				if (SectionName.Equals(PaletteNames[k].ToString(), ESearchCase::IgnoreCase)
+				 || SectionName.Equals(GetToolPaletteDisplayName(PaletteNames[k]).ToString(), ESearchCase::IgnoreCase))
 				{
 					NewPaletteNames.Add(PaletteNames[k]);
 					PaletteNames.RemoveAt(k);
@@ -755,7 +756,8 @@ void FModelingToolsEditorModeToolkit::InvokeUI()
 	{
 		for (FModelingModeCustomSectionColor ToolColor : UISettings->SectionColors)
 		{
-			if (ToolColor.SectionName.Compare(ToolbarRow.PaletteName.ToString(), ESearchCase::IgnoreCase) == 0)
+			if (ToolColor.SectionName.Equals(ToolbarRow.PaletteName.ToString(), ESearchCase::IgnoreCase)
+			 || ToolColor.SectionName.Equals(ToolbarRow.DisplayName.ToString(), ESearchCase::IgnoreCase))
 			{
 				// code below is highly dependent on the structure of the ToolbarRow.ToolbarWidget. Currently this is 
 				// a SMultiBoxWidget, a few levels below a SExpandableArea. The SExpandableArea contains a SVerticalBox
diff --git a/Engine/Plugins/Editor/ModelingToolsEditorMode/Source/ModelingToolsEditorMode/Private/ModelingToolsManagerActions.cpp b/Engine/Plugins/Editor/ModelingToolsEditorMode/Source/ModelingToolsEditorMode/Private/ModelingToolsManagerActions.cpp
--- a/Engine/Plugins/Editor/ModelingToolsEditorMode/Source/ModelingToolsEditorMode/Private/ModelingToolsManagerActions.cpp
+++ b/Engine/Plugins/Editor/ModelingToolsEditorMode/Source/ModelingToolsEditorMode/Private/ModelingToolsManagerActions.cpp
@@ -26,7 +26,8 @@ TSharedPtr<FUICommandInfo> FModelingToolsManagerCommands::FindToolByName(FString
 	bFound = false;
 	for (const FStartToolCommand& Command : RegisteredTools)
 	{
-		if (Command.ToolUIName.Compare(Name, ESearchCase::IgnoreCase) == 0)
+		if (Command.ToolUIName.Equals(Name, ESearchCase::IgnoreCase)
+		 || (Command.ToolCommand.IsValid() && Command.ToolCommand->GetLabel().ToString().Equals(Name, ESearchCase::IgnoreCase)))
 		{
 			bFound = true;
 			return Command.ToolCommand;
diff --git a/Engine/Plugins/Experimental/ChaosEditor/Source/FractureEditor/FractureEditor.Build.cs b/Engine/Plugins/Experimental/ChaosEditor/Source/FractureEditor/FractureEditor.Build.cs
--- a/Engine/Plugins/Experimental/ChaosEditor/Source/FractureEditor/FractureEditor.Build.cs
+++ b/Engine/Plugins/Experimental/ChaosEditor/Source/FractureEditor/FractureEditor.Build.cs
@@ -57,6 +57,7 @@ public class FractureEditor : ModuleRules
 				"PlanarCut",
 				"Chaos",
 				"ToolWidgets",
+				"DeveloperSettings",
 
 				// ... add private dependencies that you statically link with here ...	
 			}
diff --git a/Engine/Plugins/Experimental/ChaosEditor/Source/FractureEditor/Private/FractureEditorCommands.cpp b/Engine/Plugins/Experimental/ChaosEditor/Source/FractureEditor/Private/FractureEditorCommands.cpp
--- a/Engine/Plugins/Experimental/ChaosEditor/Source/FractureEditor/Private/FractureEditorCommands.cpp
+++ b/Engine/Plugins/Experimental/ChaosEditor/Source/FractureEditor/Private/FractureEditorCommands.cpp
@@ -13,15 +13,40 @@ FFractureEditorCommands::FFractureEditorCommands()
 {
 }
 
+TSharedPtr<FUICommandInfo> FFractureEditorCommands::FindToolByName(const FString& Name, bool& bFound) const
+{
+	bFound = false;
+	for (const FToolCommandInfo& Command : RegisteredTools)
+	{
+		if (Command.ToolUIName.ToString().Equals(Name, ESearchCase::IgnoreCase)
+		 || FTextInspector::GetSourceString(Command.ToolUIName)->Equals(Name, ESearchCase::IgnoreCase))
+		{
+			bFound = true;
+			return Command.ToolCommand;
+		}
+	}
+	return TSharedPtr<FUICommandInfo>();
+}
+
 void FFractureEditorCommands::RegisterCommands()
 {
 
+	auto RegisterCommandInfo = [this](const TSharedPtr<FUICommandInfo>& CommandInfo) -> void
+	{
+		RegisteredTools.Add(FToolCommandInfo{ CommandInfo->GetLabel(), CommandInfo });
+	};
+
 	// View settings
 	UI_COMMAND(ToggleShowBoneColors, "Colors", "Toggle Show Bone Colors", EUserInterfaceActionType::ToggleButton, FInputChord(EModifierKey::Shift, EKeys::B) );
+	RegisterCommandInfo(ToggleShowBoneColors);
 	UI_COMMAND(ViewUpOneLevel, "ViewUpOneLevel", "View Up One Level", EUserInterfaceActionType::Button, FInputChord(EModifierKey::Shift, EKeys::W) );
+	RegisterCommandInfo(ViewUpOneLevel);
 	UI_COMMAND(ViewDownOneLevel, "ViewDownOneLevel", "View Down One Level", EUserInterfaceActionType::Button, FInputChord(EModifierKey::Shift, EKeys::S) );
+	RegisterCommandInfo(ViewDownOneLevel);
 	UI_COMMAND(ExplodeMore, "ExplodeMore", "Explode 10% More", EUserInterfaceActionType::Button, FInputChord(EModifierKey::Shift, EKeys::E) );
+	RegisterCommandInfo(ExplodeMore);
 	UI_COMMAND(ExplodeLess, "ExplodeLess", "Explode 10% Less", EUserInterfaceActionType::Button, FInputChord(EModifierKey::Shift, EKeys::Q) );
+	RegisterCommandInfo(ExplodeLess);
 
 	// Fracture Tools
 	for (TObjectIterator<UClass> ClassIterator; ClassIterator; ++ClassIterator)
@@ -31,6 +56,7 @@ void FFractureEditorCommands::RegisterCommands()
 			TSubclassOf<UFractureActionTool> SubclassOf = (*ClassIterator);
 			UFractureActionTool* FractureTool = SubclassOf->GetDefaultObject<UFractureActionTool>();
 			FractureTool->RegisterUICommand(this);
+			RegisterCommandInfo(FractureTool->GetUICommandInfo());
 		}
 	}
 
diff --git a/Engine/Plugins/Experimental/ChaosEditor/Source/FractureEditor/Private/FractureEditorCommands.h b/Engine/Plugins/Experimental/ChaosEditor/Source/FractureEditor/Private/FractureEditorCommands.h
--- a/Engine/Plugins/Experimental/ChaosEditor/Source/FractureEditor/Private/FractureEditorCommands.h
+++ b/Engine/Plugins/Experimental/ChaosEditor/Source/FractureEditor/Private/FractureEditorCommands.h
@@ -10,7 +10,19 @@ class FFractureEditorCommands : public TCommands<FFractureEditorCommands>
 
 		virtual void RegisterCommands() override;
 
+	protected:
+		struct FToolCommandInfo
+		{
+			FText ToolUIName;
+			TSharedPtr<FUICommandInfo> ToolCommand;
+		};
+		TArray<FToolCommandInfo> RegisteredTools;		// Tool commands listed below are stored in this list
+
 	public:
+		/**
+		 * Find Tool start-command below by registered name (tool icon name in Mode palette)
+		 */
+		TSharedPtr<FUICommandInfo> FindToolByName(const FString& Name, bool& bFound) const;
 		
 		// Selection Commands
 		TSharedPtr< FUICommandInfo > SelectAll;
diff --git a/Engine/Plugins/Experimental/ChaosEditor/Source/FractureEditor/Private/FractureEditorModeToolkit.cpp b/Engine/Plugins/Experimental/ChaosEditor/Source/FractureEditor/Private/FractureEditorModeToolkit.cpp
--- a/Engine/Plugins/Experimental/ChaosEditor/Source/FractureEditor/Private/FractureEditorModeToolkit.cpp
+++ b/Engine/Plugins/Experimental/ChaosEditor/Source/FractureEditor/Private/FractureEditorModeToolkit.cpp
@@ -337,6 +337,58 @@ void FFractureEditorModeToolkit::InvokeUI()
 		HierarchyTab = ModeUILayerPtr->GetTabManager()->TryInvokeTab(FAssetEditorModeUILayer::TopRightTabID);
 		StatisticsTab = ModeUILayerPtr->GetTabManager()->TryInvokeTab(FAssetEditorModeUILayer::BottomLeftTabID);
 	}
+
+
+	//
+	// Apply custom section header colors.
+	// See comments below, this is done via directly manipulating Slate widgets generated deep inside BaseToolkit.cpp,
+	// and will stop working if the Slate widget structure changes
+	//
+
+	UFractureModeCustomizationSettings* UISettings = GetMutableDefault<UFractureModeCustomizationSettings>();
+
+	// look up default radii for palette toolbar expandable area headers
+	FVector4 HeaderRadii(4, 4, 0, 0);
+	const FSlateBrush* BaseBrush = FAppStyle::Get().GetBrush("PaletteToolbar.ExpandableAreaHeader");
+	if (BaseBrush != nullptr)
+	{
+		HeaderRadii = BaseBrush->OutlineSettings.CornerRadii;
+	}
+
+	for (FEdModeToolbarRow& ToolbarRow : ActiveToolBarRows)
+	{
+		for (FFractureModeCustomSectionColor ToolColor : UISettings->SectionColors)
+		{
+			if (ToolColor.SectionName.Equals(ToolbarRow.DisplayName.ToString(), ESearchCase::IgnoreCase)
+			 || ToolColor.SectionName.Equals(ToolbarRow.PaletteName.ToString(), ESearchCase::IgnoreCase))
+			{
+				// code below is highly dependent on the structure of the ToolbarRow.ToolbarWidget. Currently this is 
+				// a SMultiBoxWidget, a few levels below a SExpandableArea. The SExpandableArea contains a SVerticalBox
+				// with the header as a SBorder in Slot 0. The code will fail gracefully if this structure changes.
+
+				TSharedPtr<SWidget> ExpanderVBoxWidget = (ToolbarRow.ToolbarWidget.IsValid() && ToolbarRow.ToolbarWidget->GetParentWidget().IsValid()) ?
+					ToolbarRow.ToolbarWidget->GetParentWidget()->GetParentWidget() : TSharedPtr<SWidget>();
+				if (ExpanderVBoxWidget.IsValid() && ExpanderVBoxWidget->GetTypeAsString().Compare(TEXT("SVerticalBox")) == 0)
+				{
+					TSharedPtr<SVerticalBox> ExpanderVBox = StaticCastSharedPtr<SVerticalBox>(ExpanderVBoxWidget);
+					if (ExpanderVBox.IsValid() && ExpanderVBox->NumSlots() > 0)
+					{
+						const TSharedRef<SWidget>& SlotWidgetRef = ExpanderVBox->GetSlot(0).GetWidget();
+						TSharedPtr<SWidget> SlotWidgetPtr(SlotWidgetRef);
+						if (SlotWidgetPtr.IsValid() && SlotWidgetPtr->GetTypeAsString().Compare(TEXT("SBorder")) == 0)
+						{
+							TSharedPtr<SBorder> TopBorder = StaticCastSharedPtr<SBorder>(SlotWidgetPtr);
+							if (TopBorder.IsValid())
+							{
+								TopBorder->SetBorderImage(new FSlateRoundedBoxBrush(FSlateColor(ToolColor.Color), HeaderRadii));
+							}
+						}
+					}
+				}
+				break;
+			}
+		}
+	}
 }
 
 TSharedRef<SDockTab> FFractureEditorModeToolkit::CreateHierarchyTab(const FSpawnTabArgs& Args)
@@ -517,11 +569,64 @@ void FFractureEditorModeToolkit::SetInitialPalette()
 	}
 }
 
+void FFractureEditorModeToolkit::GetToolPaletteNames(TArray<FName>& PaletteNamesOut) const
+{
+	PaletteNamesOut = PaletteNames;
+
+	UFractureModeCustomizationSettings* UISettings = GetMutableDefault<UFractureModeCustomizationSettings>();
+
+	// if user has provided custom ordering of tool palettes in the Editor Settings, try to apply them
+	if (UISettings->ToolSectionOrder.Num() > 0)
+	{
+		TArray<FName> NewPaletteNames;
+		for (FString SectionName : UISettings->ToolSectionOrder)
+		{
+			for (int32 k = 0; k < PaletteNamesOut.Num(); ++k)
+			{
+				if (SectionName.Equals(GetToolPaletteDisplayName(PaletteNamesOut[k]).ToString(), ESearchCase::IgnoreCase)
+				 || SectionName.Equals(PaletteNamesOut[k].ToString(), ESearchCase::IgnoreCase))
+				{
+					NewPaletteNames.Add(PaletteNamesOut[k]);
+					PaletteNamesOut.RemoveAt(k);
+					break;
+				}
+			}
+		}
+		NewPaletteNames.Append(PaletteNamesOut);
+		PaletteNamesOut = MoveTemp(NewPaletteNames);
+	}
+
+	// if user has provided a list of favorite tools, add that palette to the list
+	if (UISettings->ToolFavorites.Num() > 0)
+	{
+		PaletteNamesOut.Insert(FName(TEXT("Favorites")), 0);
+	}
+}
+
 void FFractureEditorModeToolkit::BuildToolPalette(FName PaletteIndex, class FToolBarBuilder& ToolbarBuilder) 
 {
 	const FFractureEditorCommands& Commands = FFractureEditorCommands::Get();
 
-	if (PaletteIndex == TEXT("Generate"))
+	if (PaletteIndex == TEXT("Favorites"))
+	{
+		UFractureModeCustomizationSettings* UISettings = GetMutableDefault<UFractureModeCustomizationSettings>();
+
+		// build Favorites tool palette
+		for (FString ToolName : UISettings->ToolFavorites)
+		{
+			bool bFound = false;
+			TSharedPtr<FUICommandInfo> FoundToolCommand = Commands.FindToolByName(ToolName, bFound);
+			if (bFound)
+			{
+				ToolbarBuilder.AddToolBarButton(FoundToolCommand);
+			}
+			else
+			{
+				UE_LOG(LogFractureTool, Display, TEXT("FractureMode: could not find Favorited Tool %s"), *ToolName);
+			}
+		}
+	}
+	else if (PaletteIndex == TEXT("Generate"))
 	{
 		ToolbarBuilder.AddToolBarButton(Commands.GenerateAsset);
 		ToolbarBuilder.AddToolBarButton(Commands.ResetAsset);
diff --git a/Engine/Plugins/Experimental/ChaosEditor/Source/FractureEditor/Public/FractureEditorMode.h b/Engine/Plugins/Experimental/ChaosEditor/Source/FractureEditor/Public/FractureEditorMode.h
--- a/Engine/Plugins/Experimental/ChaosEditor/Source/FractureEditor/Public/FractureEditorMode.h
+++ b/Engine/Plugins/Experimental/ChaosEditor/Source/FractureEditor/Public/FractureEditorMode.h
@@ -6,6 +6,7 @@
 #include "Tools/LegacyEdModeWidgetHelpers.h"
 #include "EditorUndoClient.h"
 #include "GeometryCollection/GeometryCollectionActor.h"
+#include "Engine/DeveloperSettings.h"
 
 #include "FractureEditorMode.generated.h"
 
@@ -71,3 +72,64 @@ private:
 	// Hack: We have to set this to work around the editor defaulting to orbit around selection and breaking our custom per-bone orbiting
 	mutable TOptional<FVector> CustomOrbitPivot;
 };
+
+
+
+
+/**
+ * Defines a color to be used for a particular Tool Palette Section
+ */
+USTRUCT()
+struct FFractureModeCustomSectionColor
+{
+	GENERATED_BODY()
+
+	/** Name of Section in Fracture Mode Tool Palette */
+	UPROPERTY(EditAnywhere, Category = "SectionColor")
+	FString SectionName = TEXT("");
+
+	/** Custom Header Color */
+	UPROPERTY(EditAnywhere, Category = "SectionColor")
+	FLinearColor Color = FLinearColor::Gray;
+};
+
+
+
+UCLASS(config = Editor)
+class FRACTUREEDITOR_API UFractureModeCustomizationSettings : public UDeveloperSettings
+{
+	GENERATED_BODY()
+
+public:
+	// UDeveloperSettings overrides
+
+	virtual FName GetContainerName() const { return FName("Editor"); }
+	virtual FName GetCategoryName() const { return FName("Plugins"); }
+	virtual FName GetSectionName() const { return FName("FractureEditor"); }
+
+	virtual FText GetSectionText() const override
+	{
+		return NSLOCTEXT("FractureEditorMode", "FractureModeSettingsName", "Fracture Mode");
+	}
+
+	virtual FText GetSectionDescription() const override
+	{
+		return NSLOCTEXT("FractureEditorMode", "FractureModeSettingsDescription", "Configure the Fracture Editor Mode plugin");
+	}
+
+public:
+
+
+	/** Add the names of Fracture Mode Tool Palette Sections to have them appear at the top of the Tool Palette, in the order listed below. */
+	UPROPERTY(config, EditAnywhere, Category = "Fracture Mode|UI Customization")
+	TArray<FString> ToolSectionOrder;
+
+	/** Tool Names listed in the array below will appear in a Favorites section at the top of the Fracture Mode Tool Palette */
+	UPROPERTY(config, EditAnywhere, Category = "Fracture Mode|UI Customization")
+	TArray<FString> ToolFavorites;
+
+	/** Custom Section Header Colors for listed Sections in the Fracture Mode Tool Palette */
+	UPROPERTY(config, EditAnywhere, Category = "Fracture Mode|UI Customization")
+	TArray<FFractureModeCustomSectionColor> SectionColors;
+};
+
diff --git a/Engine/Plugins/Experimental/ChaosEditor/Source/FractureEditor/Public/FractureEditorModeToolkit.h b/Engine/Plugins/Experimental/ChaosEditor/Source/FractureEditor/Public/FractureEditorModeToolkit.h
--- a/Engine/Plugins/Experimental/ChaosEditor/Source/FractureEditor/Public/FractureEditorModeToolkit.h
+++ b/Engine/Plugins/Experimental/ChaosEditor/Source/FractureEditor/Public/FractureEditorModeToolkit.h
@@ -146,7 +146,7 @@ public:
 
 	/** Returns the number of Mode specific tabs in the mode toolbar **/ 
 	const static TArray<FName> PaletteNames;
-	virtual void GetToolPaletteNames( TArray<FName>& InPaletteName ) const { InPaletteName = PaletteNames; }
+	virtual void GetToolPaletteNames(TArray<FName>& InPaletteName) const;
 	virtual FText GetToolPaletteDisplayName(FName PaletteName) const; 
 
 	/* Exclusive Tool Palettes only allow users to use tools from one palette at a time */
