fix transforms in geometry collection -> mesh conversion tool
+ support using the override RestTransform array in the geometry collection component if it is present

#rnx

diff --git a/Engine/Plugins/Experimental/ChaosEditor/Source/FractureEditor/Private/FractureToolConvert.cpp b/Engine/Plugins/Experimental/ChaosEditor/Source/FractureEditor/Private/FractureToolConvert.cpp
--- a/Engine/Plugins/Experimental/ChaosEditor/Source/FractureEditor/Private/FractureToolConvert.cpp
+++ b/Engine/Plugins/Experimental/ChaosEditor/Source/FractureEditor/Private/FractureToolConvert.cpp
@@ -135,7 +135,7 @@ namespace UE
 namespace FractureToolConvertInternal
 {
 // Create + position a single static mesh asset in the provided asset package
-AActor* CreateMeshAsset(FGeometryCollection& Collection, const FTransform& CollectionToWorld, const TConstArrayView<UMaterialInterface*>& Materials, TConstArrayView<int32> Bones, UPackage* AssetPackage, FString UniqueAssetName, bool bCenterPivot)
+AActor* CreateMeshAsset(FGeometryCollection& Collection, const TManagedArray<FTransform>& BoneTransforms, const FTransform& CollectionToWorld, const TConstArrayView<UMaterialInterface*>& Materials, TConstArrayView<int32> Bones, UPackage* AssetPackage, FString UniqueAssetName, bool bCenterPivot)
 {
 	UWorld* TargetWorld = GEditor->GetEditorWorldContext().World();
 	check(TargetWorld);
@@ -158,7 +158,7 @@ AActor* CreateMeshAsset(FGeometryCollection& Collection, const FTransform& Colle
 
 
 	FTransform BonesToCollection;
-	ConvertToMeshDescription(*OutputMeshDescription, BonesToCollection, bCenterPivot, Collection, Bones);
+	ConvertToMeshDescription(*OutputMeshDescription, BonesToCollection, bCenterPivot, Collection, BoneTransforms, Bones);
 
 	// add a material slot. Must always have one material slot.
 	int AddMaterialCount = FMath::Max(1, Materials.Num());
@@ -206,7 +206,7 @@ AActor* CreateMeshAsset(FGeometryCollection& Collection, const FTransform& Colle
 }
 
 // Convert and save all the requested static mesh assets
-bool ConvertAndSaveMeshes(FGeometryCollection& Collection, const FTransform& CollectionToWorld, const TConstArrayView<UMaterialInterface*>& Materials, TConstArrayView<int32> Bones, FString ObjectBaseName, const UObject* RelativeToAsset, bool bPromptToSave, bool bSaveCombined, bool bCenterPivots)
+bool ConvertAndSaveMeshes(FGeometryCollection& Collection, const TManagedArray<FTransform>& BoneTransforms, const FTransform& CollectionToWorld, const TConstArrayView<UMaterialInterface*>& Materials, TConstArrayView<int32> Bones, FString ObjectBaseName, const UObject* RelativeToAsset, bool bPromptToSave, bool bSaveCombined, bool bCenterPivots)
 {
 	check(RelativeToAsset);
 
@@ -266,7 +266,7 @@ bool ConvertAndSaveMeshes(FGeometryCollection& Collection, const FTransform& Col
 		FString UniqueAssetName;
 		UPackage* AssetPackage = MakeUniquePackage(ObjectBaseName, UniqueAssetName);
 		SavePackage.Add(AssetPackage);
-		CreateMeshAsset(Collection, CollectionToWorld, Materials, Bones, AssetPackage, UniqueAssetName, bCenterPivots);
+		CreateMeshAsset(Collection, BoneTransforms, CollectionToWorld, Materials, Bones, AssetPackage, UniqueAssetName, bCenterPivots);
 	}
 	else
 	{
@@ -284,7 +284,7 @@ bool ConvertAndSaveMeshes(FGeometryCollection& Collection, const FTransform& Col
 
 			TArrayView<const int32> SingleBoneView(&UseBone, 1);
 
-			CreateMeshAsset(Collection, CollectionToWorld, Materials, SingleBoneView, AssetPackage, UniqueAssetName, bCenterPivots);
+			CreateMeshAsset(Collection, BoneTransforms, CollectionToWorld, Materials, SingleBoneView, AssetPackage, UniqueAssetName, bCenterPivots);
 		}
 	}
 
@@ -302,6 +302,7 @@ int32 UFractureToolConvert::ExecuteFracture(const FFractureToolContext& Fracture
 	if (FractureContext.IsValid())
 	{
 		FGeometryCollection& Collection = *FractureContext.GetGeometryCollection();
+		const TManagedArray<FTransform>& BoneTransforms = FractureContext.GetGeometryCollectionComponent()->RestTransforms;
 
 		FMeshDescription Mesh;
 		FTransform MeshTransform;
@@ -325,7 +326,7 @@ int32 UFractureToolConvert::ExecuteFracture(const FFractureToolContext& Fracture
 		// choose default mesh name based on corresponding geometry collection name
 		FString BaseName = FString::Printf(TEXT("%s_SM"), *FractureContext.GetFracturedGeometryCollection()->GetName());
 		
-		UE::FractureToolConvertInternal::ConvertAndSaveMeshes(Collection, FractureContext.GetTransform(), Materials,
+		UE::FractureToolConvertInternal::ConvertAndSaveMeshes(Collection, BoneTransforms, FractureContext.GetTransform(), Materials,
 			FractureContext.GetSelection(), BaseName, FractureContext.GetFracturedGeometryCollection(), 
 			ConvertSettings->bPromptForBaseName, !ConvertSettings->bPerBone, ConvertSettings->bCenterPivots);
 	}
diff --git a/Engine/Plugins/Experimental/PlanarCutPlugin/Source/PlanarCut/Private/GeometryMeshConversion.cpp b/Engine/Plugins/Experimental/PlanarCutPlugin/Source/PlanarCut/Private/GeometryMeshConversion.cpp
--- a/Engine/Plugins/Experimental/PlanarCutPlugin/Source/PlanarCut/Private/GeometryMeshConversion.cpp
+++ b/Engine/Plugins/Experimental/PlanarCutPlugin/Source/PlanarCut/Private/GeometryMeshConversion.cpp
@@ -2015,7 +2015,7 @@ void FCellMeshes::CreateMeshesForSinglePlane(const FPlanarCells& Cells, const FA
 }
 
 
-void FDynamicMeshCollection::Init(const FGeometryCollection* Collection, const TArrayView<const int32>& TransformIndices, FTransform TransformCollection, bool bSaveIsolatedVertices)
+void FDynamicMeshCollection::Init(const FGeometryCollection* Collection, const TManagedArray<FTransform>& Transforms, const TArrayView<const int32>& TransformIndices, FTransform TransformCollection, bool bSaveIsolatedVertices)
 {
 	int32 NumUVLayers = Collection->NumUVLayers();
 	Meshes.Reset();
@@ -2031,7 +2031,7 @@ void FDynamicMeshCollection::Init(const FGeometryCollection* Collection, const T
 		}
 
 		using FTransform3d = UE::Geometry::FTransform3d;
-		FTransform3d CollectionToLocal = FTransform3d(GeometryCollectionAlgo::GlobalMatrix(Collection->Transform, Collection->Parent, TransformIdx) * TransformCollection);
+		FTransform3d CollectionToLocal = FTransform3d(GeometryCollectionAlgo::GlobalMatrix(Transforms, Collection->Parent, TransformIdx) * TransformCollection);
 
 		int32 AddedMeshIdx = Meshes.Add(new FMeshData(NumUVLayers));
 		FMeshData& MeshData = Meshes[AddedMeshIdx];
diff --git a/Engine/Plugins/Experimental/PlanarCutPlugin/Source/PlanarCut/Private/GeometryMeshConversion.h b/Engine/Plugins/Experimental/PlanarCutPlugin/Source/PlanarCut/Private/GeometryMeshConversion.h
--- a/Engine/Plugins/Experimental/PlanarCutPlugin/Source/PlanarCut/Private/GeometryMeshConversion.h
+++ b/Engine/Plugins/Experimental/PlanarCutPlugin/Source/PlanarCut/Private/GeometryMeshConversion.h
@@ -201,7 +201,17 @@ struct PLANARCUT_API FDynamicMeshCollection
 		Init(Collection, TransformIndices, TransformCollection, bSaveIsolatedVertices);
 	}
 
-	void Init(const FGeometryCollection* Collection, const TArrayView<const int32>& TransformIndices, FTransform TransformCollection, bool bSaveIsolatedVertices = false);
+	FDynamicMeshCollection(const FGeometryCollection* Collection, const TManagedArray<FTransform>& Transforms, const TArrayView<const int32>& TransformIndices, FTransform TransformCollection, bool bSaveIsolatedVertices = false)
+	{
+		Init(Collection, Transforms, TransformIndices, TransformCollection, bSaveIsolatedVertices);
+	}
+
+	void Init(const FGeometryCollection* Collection, const TArrayView<const int32>& TransformIndices, FTransform TransformCollection, bool bSaveIsolatedVertices = false)
+	{
+		Init(Collection, Collection->Transform, TransformIndices, TransformCollection, bSaveIsolatedVertices);
+	}
+
+	void Init(const FGeometryCollection* Collection, const TManagedArray<FTransform>& Transforms, const TArrayView<const int32>& TransformIndices, FTransform TransformCollection, bool bSaveIsolatedVertices = false);
 
 	int32 CutWithMultiplePlanes(
 		const TArrayView<const FPlane>& Planes,
diff --git a/Engine/Plugins/Experimental/PlanarCutPlugin/Source/PlanarCut/Private/PlanarCut.cpp b/Engine/Plugins/Experimental/PlanarCutPlugin/Source/PlanarCut/Private/PlanarCut.cpp
--- a/Engine/Plugins/Experimental/PlanarCutPlugin/Source/PlanarCut/Private/PlanarCut.cpp
+++ b/Engine/Plugins/Experimental/PlanarCutPlugin/Source/PlanarCut/Private/PlanarCut.cpp
@@ -1280,13 +1280,14 @@ void ConvertToMeshDescription(
 	FTransform& TransformOut,
 	bool bCenterPivot,
 	FGeometryCollection& Collection,
+	const TManagedArray<FTransform>& BoneTransforms,
 	const TArrayView<const int32>& TransformIndices
 )
 {
 	FTransform CellsToWorld = FTransform::Identity;
 	TransformOut = FTransform::Identity;
 
-	FDynamicMeshCollection MeshCollection(&Collection, TransformIndices, CellsToWorld);
+	FDynamicMeshCollection MeshCollection(&Collection, BoneTransforms.Num() ? BoneTransforms : Collection.Transform, TransformIndices, CellsToWorld);
 	
 	FDynamicMesh3 CombinedMesh;
 	SetGeometryCollectionAttributes(CombinedMesh, Collection.NumUVLayers());
@@ -1298,13 +1299,6 @@ void ConvertToMeshDescription(
 		FDynamicMesh3& Mesh = MeshCollection.Meshes[MeshIdx].AugMesh;
 		const FTransform& FromCollection = MeshCollection.Meshes[MeshIdx].FromCollection;
 
-		// transform from the local geometry to the collection space
-		// unless it's just one mesh that we're going to center later anyway
-		if (!bCenterPivot || NumMeshes > 1)
-		{
-			MeshTransforms::ApplyTransformInverse(Mesh, (UE::Geometry::FTransform3d)FromCollection);
-		}
-
 		FMeshNormals::InitializeOverlayToPerVertexNormals(Mesh.Attributes()->PrimaryNormals(), true);
 		AugmentedDynamicMesh::InitializeOverlayToPerVertexUVs(Mesh, Collection.NumUVLayers());
 		AugmentedDynamicMesh::InitializeOverlayToPerVertexTangents(Mesh);
diff --git a/Engine/Plugins/Experimental/PlanarCutPlugin/Source/PlanarCut/Public/PlanarCut.h b/Engine/Plugins/Experimental/PlanarCutPlugin/Source/PlanarCut/Public/PlanarCut.h
--- a/Engine/Plugins/Experimental/PlanarCutPlugin/Source/PlanarCut/Public/PlanarCut.h
+++ b/Engine/Plugins/Experimental/PlanarCutPlugin/Source/PlanarCut/Public/PlanarCut.h
@@ -335,6 +335,7 @@ void PLANARCUT_API ConvertToMeshDescription(
 	FTransform& TransformOut,
 	bool bCenterPivot,
 	FGeometryCollection& Collection,
+	const TManagedArray<FTransform>& BoneTransforms,
 	const TArrayView<const int32>& TransformIndices
 );
 
