add grid distribution option for fracture mesh cut

diff --git a/Engine/Plugins/Experimental/ChaosEditor/Source/FractureEditor/Private/FractureToolMeshCut.cpp b/Engine/Plugins/Experimental/ChaosEditor/Source/FractureEditor/Private/FractureToolMeshCut.cpp
--- a/Engine/Plugins/Experimental/ChaosEditor/Source/FractureEditor/Private/FractureToolMeshCut.cpp
+++ b/Engine/Plugins/Experimental/ChaosEditor/Source/FractureEditor/Private/FractureToolMeshCut.cpp
@@ -116,10 +116,45 @@ void UFractureToolMeshCut::GenerateMeshTransforms(const FFractureToolContext& Co
 	FBox Bounds = Context.GetWorldBounds();
 	const FVector Extent(Bounds.Max - Bounds.Min);
 
-	MeshTransforms.Reserve(MeshTransforms.Num() + MeshCutSettings->NumberToScatter);
-	for (int32 ii = 0; ii < MeshCutSettings->NumberToScatter; ++ii)
+	TArray<FVector> Positions;
+	if (MeshCutSettings->CutDistribution == EMeshCutDistribution::UniformRandom)
+	{
+		Positions.Reserve(MeshCutSettings->NumberToScatter);
+		for (int32 Idx = 0; Idx < MeshCutSettings->NumberToScatter; ++Idx)
+		{
+			Positions.Emplace(Bounds.Min + FVector(RandStream.FRand(), RandStream.FRand(), RandStream.FRand()) * Extent);
+		}
+	}
+	else if (MeshCutSettings->CutDistribution == EMeshCutDistribution::Grid)
+	{
+		Positions.Reserve(MeshCutSettings->GridX * MeshCutSettings->GridY * MeshCutSettings->GridZ);
+		auto ToFrac = [](int32 Val, int32 NumVals) -> FVector::FReal
+		{
+			return (FVector::FReal(Val) + FVector::FReal(.5)) / FVector::FReal(NumVals);
+		};
+		for (int32 X = 0; X < MeshCutSettings->GridX; ++X)
+		{
+			FVector::FReal XFrac = ToFrac(X, MeshCutSettings->GridX);
+			for (int32 Y = 0; Y < MeshCutSettings->GridY; ++Y)
+			{
+				FVector::FReal YFrac = ToFrac(Y, MeshCutSettings->GridY);
+				for (int32 Z = 0; Z < MeshCutSettings->GridZ; ++Z)
+				{
+					FVector::FReal ZFrac = ToFrac(Z, MeshCutSettings->GridZ);
+					Positions.Emplace(Bounds.Min + FVector(XFrac, YFrac, ZFrac) * Extent);
+				}
+			}
+		}
+
+		for (FVector& Position : Positions)
+		{
+			Position += (RandStream.VRand() * RandStream.FRand() * MeshCutSettings->Variability);
+		}
+	}
+
+	MeshTransforms.Reserve(MeshTransforms.Num() + Positions.Num());
+	for (const FVector& Position : Positions)
 	{
-		FVector Position(Bounds.Min + FVector(RandStream.FRand(), RandStream.FRand(), RandStream.FRand()) * Extent);
 		FVector::FReal Scale = RandStream.FRandRange(MeshCutSettings->MinScaleFactor, MeshCutSettings->MaxScaleFactor);
 		FVector ScaleVec(Scale, Scale, Scale);
 		FRotator Orientation = FRotator::ZeroRotator;
@@ -176,7 +211,7 @@ int32 UFractureToolMeshCut::ExecuteFracture(const FFractureToolContext& Fracture
 		ClearProximity(FractureContext.GetGeometryCollection().Get());
 
 		// (Note: noise and grout not currently supported)
-		if (LocalCutSettings->NumberToScatter == 0)
+		if (LocalCutSettings->CutDistribution == EMeshCutDistribution::SingleCut)
 		{
 			return CutWithMesh(MeshDescription, ActorTransform, InternalSurfaceMaterials, *FractureContext.GetGeometryCollection(), FractureContext.GetSelection(), CollisionSettings->PointSpacing, FractureContext.GetTransform());
 		}
diff --git a/Engine/Plugins/Experimental/ChaosEditor/Source/FractureEditor/Private/FractureToolMeshCut.h b/Engine/Plugins/Experimental/ChaosEditor/Source/FractureEditor/Private/FractureToolMeshCut.h
--- a/Engine/Plugins/Experimental/ChaosEditor/Source/FractureEditor/Private/FractureToolMeshCut.h
+++ b/Engine/Plugins/Experimental/ChaosEditor/Source/FractureEditor/Private/FractureToolMeshCut.h
@@ -9,6 +9,16 @@
 
 class FFractureToolContext;
 
+UENUM()
+enum class EMeshCutDistribution
+{
+	// Cut only once, at the cutting mesh's current location in the level
+	SingleCut,
+	// Scatter the cutting mesh in a uniform random distribution around the geometry bounding box
+	UniformRandom,
+	// Arrange the cutting mesh in a regular grid pattern
+	Grid
+};
 
 UCLASS(config = EditorPerProjectUserSettings)
 class UFractureMeshCutSettings : public UFractureToolSettings
@@ -23,29 +33,53 @@ public:
 	UPROPERTY(EditAnywhere, Category = MeshCut, meta = (DisplayName = "Cutting Actor"))
 	TLazyObjectPtr<AStaticMeshActor> CuttingActor;
 
-	/** Number of meshes to random scatter; if 0, the selected mesh will be used in its current position, with no random scattering*/
+	/** How to arrange the mesh cuts in space */
 	UPROPERTY(EditAnywhere, Category = Distribution)
-	int NumberToScatter = 0;
+	EMeshCutDistribution CutDistribution = EMeshCutDistribution::SingleCut;
+
+	/** Number of meshes to random scatter */
+	UPROPERTY(EditAnywhere, Category = Distribution, meta = (ClampMin = "1", UIMax = "5000",
+		EditCondition = "CutDistribution == EMeshCutDistribution::UniformRandom", EditConditionHides))
+	int NumberToScatter = 10;
+
+	/** Number of meshes to add to grid in X */
+	UPROPERTY(EditAnywhere, Category = Distribution, meta = (DisplayName = "Grid Width", ClampMin = "1", UIMax = "100", ClampMax = "5000",
+		EditCondition = "CutDistribution == EMeshCutDistribution::Grid", EditConditionHides))
+	int32 GridX = 5;
+
+	/** Number of meshes to add to grid in Y */
+	UPROPERTY(EditAnywhere, Category = Distribution, meta = (DisplayName = "Grid Depth", ClampMin = "1", UIMax = "100", ClampMax = "5000",
+		EditCondition = "CutDistribution == EMeshCutDistribution::Grid", EditConditionHides))
+	int32 GridY = 5;
+
+	/** Number of meshes to add to grid in Z */
+	UPROPERTY(EditAnywhere, Category = Distribution, meta = (DisplayName = "Grid Height", ClampMin = "1", UIMax = "100", ClampMax = "5000",
+		EditCondition = "CutDistribution == EMeshCutDistribution::Grid", EditConditionHides))
+	int32 GridZ = 5;
+
+	/** Magnitude of random displacement to cutting meshes */
+	UPROPERTY(EditAnywhere, Category = Distribution, meta = (DisplayName = "Variability", EditCondition = "CutDistribution == EMeshCutDistribution::Grid", EditConditionHides, UIMin = "0.0", ClampMin = "0.0"))
+	float Variability = 0.0f;
 
-	UPROPERTY(EditAnywhere, Category = Distribution, meta = (ClampMin = "0.001", EditCondition = "NumberToScatter > 0"))
+	UPROPERTY(EditAnywhere, Category = Distribution, meta = (ClampMin = "0.001", EditCondition = "CutDistribution != EMeshCutDistribution::SingleCut", EditConditionHides))
 	float MinScaleFactor = .5;
 
-	UPROPERTY(EditAnywhere, Category = Distribution, meta = (ClampMin = "0.001", EditCondition = "NumberToScatter > 0"))
+	UPROPERTY(EditAnywhere, Category = Distribution, meta = (ClampMin = "0.001", EditCondition = "CutDistribution != EMeshCutDistribution::SingleCut", EditConditionHides))
 	float MaxScaleFactor = 1.5;
 
-	UPROPERTY(EditAnywhere, Category = Distribution, meta = (EditCondition = "NumberToScatter > 0"))
+	UPROPERTY(EditAnywhere, Category = Distribution, meta = (EditCondition = "CutDistribution != EMeshCutDistribution::SingleCut", EditConditionHides))
 	bool bRandomOrientation = true;
 
 	/** Roll will be chosen between -Range to +Range */
-	UPROPERTY(EditAnywhere, Category = Distribution, meta = (DisplayName = "+/- Roll Range", EditCondition = "NumberToScatter > 0 && bRandomOrientation", ClampMin = "0", ClampMax = "180"))
+	UPROPERTY(EditAnywhere, Category = Distribution, meta = (DisplayName = "+/- Roll Range", EditCondition = "CutDistribution != EMeshCutDistribution::SingleCut && bRandomOrientation", EditConditionHides, ClampMin = "0", ClampMax = "180"))
 	float RollRange = 180;
 
 	/** Pitch will be chosen between -Range to +Range */
-	UPROPERTY(EditAnywhere, Category = Distribution, meta = (DisplayName = "+/- Pitch Range", EditCondition = "NumberToScatter > 0 && bRandomOrientation", ClampMin = "0", ClampMax = "180"))
+	UPROPERTY(EditAnywhere, Category = Distribution, meta = (DisplayName = "+/- Pitch Range", EditCondition = "CutDistribution != EMeshCutDistribution::SingleCut && bRandomOrientation", EditConditionHides, ClampMin = "0", ClampMax = "180"))
 	float PitchRange = 180;
 
 	/** Yaw will be chosen between -Range to +Range */
-	UPROPERTY(EditAnywhere, Category = Distribution, meta = (DisplayName = "+/- Yaw Range", EditCondition = "NumberToScatter > 0 && bRandomOrientation", ClampMin = "0", ClampMax = "180"))
+	UPROPERTY(EditAnywhere, Category = Distribution, meta = (DisplayName = "+/- Yaw Range", EditCondition = "CutDistribution != EMeshCutDistribution::SingleCut && bRandomOrientation", EditConditionHides, ClampMin = "0", ClampMax = "180"))
 	float YawRange = 180;
 };
 
