fix crash on updating explodedview when geometrycollection is malformed (has embedded parented to embedded, or to nothing)

#rnx

diff --git a/Engine/Plugins/Experimental/ChaosEditor/Source/FractureEditor/Private/FractureEditorModeToolkit.cpp b/Engine/Plugins/Experimental/ChaosEditor/Source/FractureEditor/Private/FractureEditorModeToolkit.cpp
--- a/Engine/Plugins/Experimental/ChaosEditor/Source/FractureEditor/Private/FractureEditorModeToolkit.cpp
+++ b/Engine/Plugins/Experimental/ChaosEditor/Source/FractureEditor/Private/FractureEditorModeToolkit.cpp
@@ -1616,7 +1616,16 @@ bool GetValidGeoCenter(FGeometryCollection* Collection, const TManagedArray<int3
 	}
 	else if (Collection->SimulationType[TransformIndex] == FGeometryCollection::ESimulationTypes::FST_None) // ie this is embedded geometry
 	{
-		OutGeoCenter = Transforms[Collection->Parent[TransformIndex]].TransformPosition(BoundingBox[TransformToGeometryIndex[Collection->Parent[TransformIndex]]].GetCenter());
+		int32 Parent = Collection->Parent[TransformIndex];
+		int32 ParentGeo = Parent != INDEX_NONE ? TransformToGeometryIndex[Parent] : INDEX_NONE;
+		if (ensureMsgf(ParentGeo != INDEX_NONE, TEXT("Embedded geometry should always have a rigid geometry parent!  Geometry collection may be malformed.")))
+		{
+			OutGeoCenter = Transforms[Collection->Parent[TransformIndex]].TransformPosition(BoundingBox[ParentGeo].GetCenter());
+		}
+		else
+		{
+			return false; // no valid value to return
+		}
 
 		return true;
 	}
