Fix Draw Sites toggle for brick fracture, fix tooltips for draw sites/draw fracture, and remove non-functioning Draw Sites toggle on slice fracture

#rnx

diff --git a/Engine/Plugins/Experimental/ChaosEditor/Source/FractureEditor/Private/FractureToolBrick.cpp b/Engine/Plugins/Experimental/ChaosEditor/Source/FractureEditor/Private/FractureToolBrick.cpp
--- a/Engine/Plugins/Experimental/ChaosEditor/Source/FractureEditor/Private/FractureToolBrick.cpp
+++ b/Engine/Plugins/Experimental/ChaosEditor/Source/FractureEditor/Private/FractureToolBrick.cpp
@@ -335,11 +335,14 @@ void UFractureToolBrick::FractureContextChanged()
 
 void UFractureToolBrick::Render(const FSceneView* View, FViewport* Viewport, FPrimitiveDrawInterface* PDI)
 {
-	EnumerateVisualizationMapping(BricksMappings, BrickTransforms.Num(), [&](int32 Idx, FVector ExplodedVector)
+	if (CutterSettings->bDrawSites)
 	{
-		const FTransform& Transform = BrickTransforms[Idx];
-		PDI->DrawPoint(Transform.GetLocation() + ExplodedVector, FLinearColor::Green, 4.f, SDPG_Foreground);
-	});
+		EnumerateVisualizationMapping(BricksMappings, BrickTransforms.Num(), [&](int32 Idx, FVector ExplodedVector)
+		{
+			const FTransform& Transform = BrickTransforms[Idx];
+			PDI->DrawPoint(Transform.GetLocation() + ExplodedVector, FLinearColor::Green, 4.f, SDPG_Foreground);
+		});
+	}
 
 	if (CutterSettings->bDrawDiagram)
 	{
diff --git a/Engine/Plugins/Experimental/ChaosEditor/Source/FractureEditor/Private/FractureToolCutter.h b/Engine/Plugins/Experimental/ChaosEditor/Source/FractureEditor/Private/FractureToolCutter.h
--- a/Engine/Plugins/Experimental/ChaosEditor/Source/FractureEditor/Private/FractureToolCutter.h
+++ b/Engine/Plugins/Experimental/ChaosEditor/Source/FractureEditor/Private/FractureToolCutter.h
@@ -55,11 +55,15 @@ public:
 	UPROPERTY()
 	bool bGroupFractureToggleEnabled = true;
 
-	/** Generate a fracture pattern across all selected meshes.  */
-	UPROPERTY(EditAnywhere, Category = CommonFracture, meta = (DisplayName = "Draw Sites"))
+	/** Draw points marking the centers of pieces to be cut out by the fracture pattern.  */
+	UPROPERTY(EditAnywhere, Category = CommonFracture, meta = (EditCondition = "bDrawSitesToggleEnabled", HideEditConditionToggle, EditConditionHides, DisplayName = "Draw Sites"))
 	bool bDrawSites;
 
-	/** Generate a fracture pattern across all selected meshes.  */
+	// This flag allows tools to disable the above bDrawSites option if/when it is not applicable
+	UPROPERTY()
+	bool bDrawSitesToggleEnabled = true;
+
+	/** Draw the edges of the fracture pattern.  */
 	UPROPERTY(EditAnywhere, Category = CommonFracture, meta = (DisplayName = "Draw Diagram"))
 	bool bDrawDiagram;
 
diff --git a/Engine/Plugins/Experimental/ChaosEditor/Source/FractureEditor/Private/FractureToolSlice.cpp b/Engine/Plugins/Experimental/ChaosEditor/Source/FractureEditor/Private/FractureToolSlice.cpp
--- a/Engine/Plugins/Experimental/ChaosEditor/Source/FractureEditor/Private/FractureToolSlice.cpp
+++ b/Engine/Plugins/Experimental/ChaosEditor/Source/FractureEditor/Private/FractureToolSlice.cpp
@@ -21,6 +21,8 @@ UFractureToolSlice::UFractureToolSlice(const FObjectInitializer& ObjInit)
 {
 	SliceSettings = NewObject<UFractureSliceSettings>(GetTransientPackage(), UFractureSliceSettings::StaticClass());
 	SliceSettings->OwnerTool = this;
+
+	CutterSettings->bDrawSitesToggleEnabled = false;
 }
 
 FText UFractureToolSlice::GetDisplayText() const
