Disable radial fracture position (center, normal) settings when position is controlled by gizmo.  Update tooltips to explain why they're disabled.

diff --git a/Engine/Plugins/Experimental/ChaosEditor/Source/FractureEditor/Private/FractureToolCutter.cpp b/Engine/Plugins/Experimental/ChaosEditor/Source/FractureEditor/Private/FractureToolCutter.cpp
--- a/Engine/Plugins/Experimental/ChaosEditor/Source/FractureEditor/Private/FractureToolCutter.cpp
+++ b/Engine/Plugins/Experimental/ChaosEditor/Source/FractureEditor/Private/FractureToolCutter.cpp
@@ -30,10 +30,6 @@ void UFractureTransformGizmoSettings::ResetGizmo(bool bResetRotation)
 	{
 		return;
 	}
-	if (AttachedCutter)
-	{
-		AttachedCutter->SetMandateGroupFracture(bUseGizmo);
-	}
 	if (!bUseGizmo || !AttachedCutter)
 	{
 		TransformGizmo->SetVisibility(false);
@@ -62,6 +58,7 @@ void UFractureTransformGizmoSettings::PostEditChangeProperty(FPropertyChangedEve
 	{
 		if (AttachedCutter)
 		{
+			AttachedCutter->UpdateUseGizmo(bUseGizmo);
 			ResetGizmo();
 		}
 	}
@@ -97,6 +94,7 @@ void UFractureTransformGizmoSettings::Setup(UFractureToolCutterBase* Cutter, ETr
 		}
 		TransformGizmo->SetActiveTarget(TransformProxy);
 		TransformProxy->OnTransformChanged.AddUObject(this, &UFractureTransformGizmoSettings::TransformChanged);
+		AttachedCutter->UpdateUseGizmo(bUseGizmo);
 		ResetGizmo();
 	}
 }
diff --git a/Engine/Plugins/Experimental/ChaosEditor/Source/FractureEditor/Private/FractureToolCutter.h b/Engine/Plugins/Experimental/ChaosEditor/Source/FractureEditor/Private/FractureToolCutter.h
--- a/Engine/Plugins/Experimental/ChaosEditor/Source/FractureEditor/Private/FractureToolCutter.h
+++ b/Engine/Plugins/Experimental/ChaosEditor/Source/FractureEditor/Private/FractureToolCutter.h
@@ -147,6 +147,12 @@ public:
 
 	FBox GetCombinedBounds(const TArray<FFractureToolContext>& Contexts) const;
 
+
+	virtual void UpdateUseGizmo(bool bUseGizmo)
+	{
+		SetMandateGroupFracture(bUseGizmo);
+	}
+
 	void SetMandateGroupFracture(bool bMandateGroupFracture)
 	{
 		CutterSettings->bGroupFractureToggleEnabled = !bMandateGroupFracture;
diff --git a/Engine/Plugins/Experimental/ChaosEditor/Source/FractureEditor/Private/FractureToolRadial.h b/Engine/Plugins/Experimental/ChaosEditor/Source/FractureEditor/Private/FractureToolRadial.h
--- a/Engine/Plugins/Experimental/ChaosEditor/Source/FractureEditor/Private/FractureToolRadial.h
+++ b/Engine/Plugins/Experimental/ChaosEditor/Source/FractureEditor/Private/FractureToolRadial.h
@@ -26,14 +26,17 @@ public:
 		, Variability(0.0f)
 	{}
 
-	/** Center of generated pattern */
-	UPROPERTY(EditAnywhere, Category = RadialVoronoi)
+	/** Center of generated pattern.  Only used when "Use Gizmo" is disabled. */
+	UPROPERTY(EditAnywhere, Category = RadialVoronoi, meta = (EditCondition = "!bPositionedByGizmo", HideEditConditionToggle))
 	FVector Center;
 
-	/** Normal to plane in which sites are generated */
-	UPROPERTY(EditAnywhere, Category = RadialVoronoi)
+	/** Normal to plane in which sites are generated.  Only used when "Use Gizmo" is disabled. */
+	UPROPERTY(EditAnywhere, Category = RadialVoronoi, meta = (EditCondition = "!bPositionedByGizmo", HideEditConditionToggle))
 	FVector Normal;
 
+	UPROPERTY()
+	bool bPositionedByGizmo = true;
+
 	/** Pattern radius */
 	UPROPERTY(EditAnywhere, Category = RadialVoronoi, meta = (DisplayName = "Radius", UIMin = "0.0", ClampMin = "0.0"))
 	float Radius;
@@ -83,6 +86,13 @@ public:
 	virtual void Setup() override;
 	virtual void Shutdown() override;
 
+	virtual void UpdateUseGizmo(bool bUseGizmo) override
+	{
+		Super::UpdateUseGizmo(bUseGizmo);
+		RadialSettings->bPositionedByGizmo = bUseGizmo;
+		NotifyOfPropertyChangeByTool(RadialSettings);
+	}
+
 protected:
 	void GenerateVoronoiSites(const FFractureToolContext& Context, TArray<FVector>& Sites) override;
 
