fix bug where fracture embedded geometry that was hidden by the 'hide unselected' checkbox would not re-show when 'hide unselected' was unchecked

diff --git a/Engine/Plugins/Experimental/ChaosEditor/Source/FractureEditor/Private/FractureEditorModeToolkit.cpp b/Engine/Plugins/Experimental/ChaosEditor/Source/FractureEditor/Private/FractureEditorModeToolkit.cpp
--- a/Engine/Plugins/Experimental/ChaosEditor/Source/FractureEditor/Private/FractureEditorModeToolkit.cpp
+++ b/Engine/Plugins/Experimental/ChaosEditor/Source/FractureEditor/Private/FractureEditorModeToolkit.cpp
@@ -1129,6 +1129,7 @@ void FFractureEditorModeToolkit::UpdateHideForComponent(UGeometryCollectionCompo
 			{
 				GeometryCollection->RemoveAttribute("Hide", FGeometryCollection::TransformGroup);
 			}
+			Comp->RefreshEmbeddedGeometry();
 		}
 	}
 }
