Add a grid to the fracture plane cut tool gizmo to make it easier to line up

#rnx

diff --git a/Engine/Plugins/Experimental/ChaosEditor/Source/FractureEditor/Private/FractureToolPlaneCut.cpp b/Engine/Plugins/Experimental/ChaosEditor/Source/FractureEditor/Private/FractureToolPlaneCut.cpp
--- a/Engine/Plugins/Experimental/ChaosEditor/Source/FractureEditor/Private/FractureToolPlaneCut.cpp
+++ b/Engine/Plugins/Experimental/ChaosEditor/Source/FractureEditor/Private/FractureToolPlaneCut.cpp
@@ -7,6 +7,8 @@
 #include "GeometryCollection/GeometryCollectionObject.h"
 #include "FractureEditorModeToolkit.h"
 #include "FractureToolContext.h"
+#include "Drawing/MeshDebugDrawing.h"
+#include "FrameTypes.h"
 
 #define LOCTEXT_NAMESPACE "FracturePlanar"
 
@@ -82,6 +84,14 @@ void UFractureToolPlaneCut::Render(const FSceneView* View, FViewport* Viewport,
 			PDI->DrawLine(Corners[1], Corners[2], FLinearColor(0, 255, 0), SDPG_Foreground);
 			PDI->DrawLine(Corners[2], Corners[3], FLinearColor(255, 0, 0), SDPG_Foreground);
 			PDI->DrawLine(Corners[3], Corners[0], FLinearColor(0, 255, 0), SDPG_Foreground);
+
+			// Put some grid lines in that square
+			const int32 NumGridLines = 10;
+			const float GridLineSpacing = 1.5f * PlaneSize / NumGridLines;
+			const FColor GridColor(128, 128, 128, 64);
+			const float GridThickness = 0.75f;
+			const UE::Geometry::FFrame3f DrawFrame(Center, Transform.GetRotation());
+			MeshDebugDraw::DrawSimpleGrid(DrawFrame, NumGridLines, GridLineSpacing, GridThickness, GridColor, false, PDI, FTransform::Identity);
 		};
 
 		if (GizmoSettings->IsGizmoEnabled())
