Fracture tool plane cut gizmo: have the grid in the gizmo be occluded by the fracture mesh, and also don't extend the grid outside of the the gizmo square.

#rnx

diff --git a/Engine/Plugins/Experimental/ChaosEditor/Source/FractureEditor/Private/FractureToolPlaneCut.cpp b/Engine/Plugins/Experimental/ChaosEditor/Source/FractureEditor/Private/FractureToolPlaneCut.cpp
--- a/Engine/Plugins/Experimental/ChaosEditor/Source/FractureEditor/Private/FractureToolPlaneCut.cpp
+++ b/Engine/Plugins/Experimental/ChaosEditor/Source/FractureEditor/Private/FractureToolPlaneCut.cpp
@@ -86,12 +86,12 @@ void UFractureToolPlaneCut::Render(const FSceneView* View, FViewport* Viewport,
 			PDI->DrawLine(Corners[3], Corners[0], FLinearColor(0, 255, 0), SDPG_Foreground);
 
 			// Put some grid lines in that square
-			const int32 NumGridLines = 10;
-			const float GridLineSpacing = 1.5f * PlaneSize / NumGridLines;
-			const FColor GridColor(128, 128, 128, 64);
-			const float GridThickness = 0.75f;
+			const int32 NumGridLines = 20;
+			const float GridLineSpacing = PlaneSize / (NumGridLines - 1);
+			const FColor GridColor(64, 64, 64, 128);
+			const float GridThickness = 1.0f;
 			const UE::Geometry::FFrame3f DrawFrame(Center, Transform.GetRotation());
-			MeshDebugDraw::DrawSimpleGrid(DrawFrame, NumGridLines, GridLineSpacing, GridThickness, GridColor, false, PDI, FTransform::Identity);
+			MeshDebugDraw::DrawSimpleGrid(DrawFrame, NumGridLines, GridLineSpacing, GridThickness, GridColor, true, PDI, FTransform::Identity);
 		};
 
 		if (GizmoSettings->IsGizmoEnabled())
