Ignore CTRL-drag in Fracture mode. Users are unintentially nudging the geometry collection while CTRL-clicking to select/deselect facture pieces.

#rnx

diff --git a/Engine/Plugins/Experimental/ChaosEditor/Source/FractureEditor/Private/FractureEditorMode.cpp b/Engine/Plugins/Experimental/ChaosEditor/Source/FractureEditor/Private/FractureEditorMode.cpp
--- a/Engine/Plugins/Experimental/ChaosEditor/Source/FractureEditor/Private/FractureEditorMode.cpp
+++ b/Engine/Plugins/Experimental/ChaosEditor/Source/FractureEditor/Private/FractureEditorMode.cpp
@@ -224,6 +224,21 @@ bool UFractureEditorMode::HandleClick(FEditorViewportClient* InViewportClient, H
 
 }
 
+
+bool UFractureEditorMode::InputAxis(FEditorViewportClient* InViewportClient, FViewport* InViewport, int32 ControllerId, FKey Key, float Delta, float DeltaTime)
+{
+	// Disable the CTRL+drag behavior of the level editor (translate/rotate the selected object along an axis), since it can easily result in 
+	// unintentionally moving objects while CTRL-clicking to select/deselect
+
+	bool bCtrlDown = InViewport->KeyState(EKeys::LeftControl) || InViewport->KeyState(EKeys::RightControl);
+	if (bCtrlDown)
+	{
+		return true;
+	}
+	return UBaseLegacyWidgetEdMode::InputAxis(InViewportClient, InViewport, ControllerId, Key, Delta, DeltaTime);
+}
+
+
 bool UFractureEditorMode::BoxSelect(FBox& InBox, bool InSelect /*= true*/)
 {
 	FConvexVolume BoxVolume(GetVolumeFromBox(InBox));
diff --git a/Engine/Plugins/Experimental/ChaosEditor/Source/FractureEditor/Public/FractureEditorMode.h b/Engine/Plugins/Experimental/ChaosEditor/Source/FractureEditor/Public/FractureEditorMode.h
--- a/Engine/Plugins/Experimental/ChaosEditor/Source/FractureEditor/Public/FractureEditorMode.h
+++ b/Engine/Plugins/Experimental/ChaosEditor/Source/FractureEditor/Public/FractureEditorMode.h
@@ -46,6 +46,9 @@ public:
 	virtual bool GetPivotForOrbit(FVector& OutPivot) const override;
 	// End of UEdMode interface
 
+	// ILegacyEdModeViewportInterface
+	virtual bool InputAxis(FEditorViewportClient* InViewportClient, FViewport* InViewport, int32 ControllerId, FKey Key, float Delta, float DeltaTime) override;
+
 	// ILegacyEdModeSelectInterface
 	virtual bool BoxSelect(FBox& InBox, bool InSelect = true) override;
 	virtual bool FrustumSelect(const FConvexVolume& InFrustum, FEditorViewportClient* InViewportClient, bool InSelect) override;
