add 'selected bone' sampling and also a 'normal offset' feature for both selected bone and selected mesh sampling.  use averaged normals per vertex to avoid creating more samples at sharp edges.

#rnx

diff --git a/Engine/Plugins/Experimental/ChaosEditor/Source/FractureEditor/Private/FractureToolCustomVoronoi.cpp b/Engine/Plugins/Experimental/ChaosEditor/Source/FractureEditor/Private/FractureToolCustomVoronoi.cpp
--- a/Engine/Plugins/Experimental/ChaosEditor/Source/FractureEditor/Private/FractureToolCustomVoronoi.cpp
+++ b/Engine/Plugins/Experimental/ChaosEditor/Source/FractureEditor/Private/FractureToolCustomVoronoi.cpp
@@ -111,7 +111,7 @@ TArray<UObject*> UFractureToolCustomVoronoi::GetSettingsObjects() const
 	return Settings;
 }
 
-void UFractureToolCustomVoronoi::GenerateLivePattern(int32 RandomSeed)
+void UFractureToolCustomVoronoi::GenerateLivePattern(const TArray<FFractureToolContext>& FractureContexts, int32 RandomSeed)
 {
 	if (!ensure(CombinedWorldBounds.IsValid))
 	{
@@ -191,14 +191,106 @@ void UFractureToolCustomVoronoi::GenerateLivePattern(int32 RandomSeed)
 
 			FMeshDescription* MeshDescription = CustomVoronoiSettings->ReferenceMesh->GetStaticMeshComponent()->GetStaticMesh()->GetMeshDescription(0);
 			TVertexAttributesConstRef<FVector3f> VertexPositions = MeshDescription->GetVertexPositions();
-			// copy vertex positions
+			TArray<FVector3f> Normals; // Per-vertex normals for offsets
+			float NormalOffset = CustomVoronoiSettings->NormalOffset;
+			if (NormalOffset > 0)
+			{
+				Normals.SetNumZeroed(VertexPositions.GetNumElements());
+
+				FStaticMeshConstAttributes Attributes(*MeshDescription);
+				TVertexInstanceAttributesConstRef<FVector3f> InstanceNormals = Attributes.GetVertexInstanceNormals();
+				for (const FVertexInstanceID InstanceID : MeshDescription->VertexInstances().GetElementIDs())
+				{
+					FVector3f Normal = InstanceNormals.Get(InstanceID);
+					FVertexID VertexID = MeshDescription->GetVertexInstanceVertex(InstanceID);
+					Normals[VertexID] += Normal;
+				}
+				for (FVector3f& Normal : Normals)
+				{
+					Normal.Normalize();
+				}
+			}
 			for (const FVertexID VertexID : MeshDescription->Vertices().GetElementIDs())
 			{
-				const FVector Position = VertexPositions.Get(VertexID);
+				FVector Position = VertexPositions.Get(VertexID);
+				if (Normals.Num() > 0)
+				{
+					Position += Normals[VertexID] * NormalOffset;
+				}
 				Sites.Add(VerticesTransform.TransformPosition(Position));
 			}
 		}
 	}
+	else if (CustomVoronoiSettings->VoronoiPattern == EVoronoiPattern::SelectedBones)
+	{
+		for (const FFractureToolContext& Context : FractureContexts)
+		{
+			FGeometryCollection& Collection = *Context.GetGeometryCollection();
+			for (int32 Bone : Context.GetSelection())
+			{
+				int32 GeometryIdx = Collection.TransformToGeometryIndex[Bone];
+				if (GeometryIdx == INDEX_NONE)
+				{
+					continue;
+				}
+				int32 Start = Collection.VertexStart[GeometryIdx];
+				int32 Count = Collection.VertexCount[GeometryIdx];
+				// for now, always exclude vertices that are not used by any triangles
+				// (by convention, these are always at the end of their section of the array)
+				constexpr bool bIncludeCollisionSampleVerts = false;
+				if (!bIncludeCollisionSampleVerts)
+				{
+					int32 MaxUsedIdx = Start;
+					int32 FStart = Collection.FaceStart[GeometryIdx];
+					int32 FCount = Collection.FaceCount[GeometryIdx];
+					for (int32 FIdx = FStart; FIdx < FStart + FCount; FIdx++)
+					{
+						MaxUsedIdx = FMath::Max(MaxUsedIdx, Collection.Indices[FIdx].GetMax());
+					}
+					Count = MaxUsedIdx - Start + 1;
+				}
+				float NormalOffset = CustomVoronoiSettings->NormalOffset;
+				if (NormalOffset > 0)
+				{
+					TArray<FVector3f> Normals;
+					Normals.SetNumZeroed(Count);
+					TMap<FVector, int32> VertexHash;
+
+					// merge exact duplicate points and average their normals
+					for (int32 VIdx = Start; VIdx < Start + Count; VIdx++)
+					{
+						const FVector& Position = Collection.Vertex[VIdx];
+						int32& Idx = VertexHash.FindOrAdd(Position);
+						if (Idx - 1 == INDEX_NONE)
+						{
+							Idx = VIdx + 1 - Start;
+						}
+						Normals[Idx - 1] += Collection.Normal[VIdx];
+					}
+					for (FVector3f& Normal : Normals)
+					{
+						Normal.Normalize();
+					}
+					for (TPair<FVector, int32>& Point : VertexHash)
+					{
+						const FVector OffsetPosition = Point.Key + Normals[Point.Value - 1] * NormalOffset;
+						const FVector LocalPos = Collection.Transform[Bone].TransformPosition(OffsetPosition);
+						Sites.Add(Context.GetTransform().TransformPosition(LocalPos));
+					}
+				}
+				else
+				{
+					// if no normal offset, allow Voronoi diagram to drop the duplicate points automatically
+					for (int32 VIdx = Start; VIdx < Start + Count; VIdx++)
+					{
+						const FVector& Position = Collection.Vertex[VIdx];
+						const FVector LocalPos = Collection.Transform[Bone].TransformPosition(Position);
+						Sites.Add(Context.GetTransform().TransformPosition(LocalPos));
+					}
+				}
+			}
+		}
+	}
 	else if (CustomVoronoiSettings->VoronoiPattern == EVoronoiPattern::Grid)
 	{
 		auto ToFrac = [](int32 Val, int32 NumVals) -> FVector::FReal
@@ -271,7 +363,7 @@ void UFractureToolCustomVoronoi::FractureContextChanged()
 
 	if (CombinedWorldBounds.IsValid && LiveSites.Num() == 0 && FractureContexts.Num() > 0)
 	{
-		GenerateLivePattern(FractureContexts[0].GetSeed());
+		GenerateLivePattern(FractureContexts, FractureContexts[0].GetSeed());
 	}
 
 	UpdateVisualizations(FractureContexts);
diff --git a/Engine/Plugins/Experimental/ChaosEditor/Source/FractureEditor/Private/FractureToolCustomVoronoi.h b/Engine/Plugins/Experimental/ChaosEditor/Source/FractureEditor/Private/FractureToolCustomVoronoi.h
--- a/Engine/Plugins/Experimental/ChaosEditor/Source/FractureEditor/Private/FractureToolCustomVoronoi.h
+++ b/Engine/Plugins/Experimental/ChaosEditor/Source/FractureEditor/Private/FractureToolCustomVoronoi.h
@@ -19,7 +19,9 @@ enum class EVoronoiPattern
 	// Add a regular grid of points
 	Grid,
 	// Add a point at every mesh vertex
-	MeshVertices
+	MeshVertices,
+	// Add a point per vertex of the selected bones
+	SelectedBones
 };
 
 UCLASS(config = EditorPerProjectUserSettings)
@@ -36,6 +38,10 @@ public:
 	/** Method to generate next group of voronoi sites */
 	UPROPERTY(EditAnywhere, Category = LiveVoronoiSites, meta = (DisplayName = "Pattern"))
 	EVoronoiPattern VoronoiPattern = EVoronoiPattern::Uniform;
+	
+	/** Offset point samples in the vertex normal directions */
+	UPROPERTY(EditAnywhere, Category = LiveVoronoiSites, meta = (EditCondition = "VoronoiPattern == EVoronoiPattern::MeshVertices || VoronoiPattern == EVoronoiPattern::SelectedBones"))
+	float NormalOffset = 0.0f;
 
 	/** Randomness of sites distribution */
 	UPROPERTY(EditAnywhere, Category = LiveVoronoiSites, meta = (DisplayName = "Variability", UIMin = "0.0", ClampMin = "0.0"))
@@ -153,7 +159,7 @@ public:
 	}
 
 	// Re-calculate the "live transformed" voronoi sites based on current settings
-	void GenerateLivePattern(int32 RandomSeed);
+	void GenerateLivePattern(const TArray<FFractureToolContext>& FractureContexts, int32 RandomSeed);
 
 	void FractureContextChanged() override;
 
