CIS fix for LWC-disabled compile, fracture tool custom voronoi
#rnx

diff --git a/Engine/Plugins/Experimental/ChaosEditor/Source/FractureEditor/Private/FractureToolCustomVoronoi.cpp b/Engine/Plugins/Experimental/ChaosEditor/Source/FractureEditor/Private/FractureToolCustomVoronoi.cpp
--- a/Engine/Plugins/Experimental/ChaosEditor/Source/FractureEditor/Private/FractureToolCustomVoronoi.cpp
+++ b/Engine/Plugins/Experimental/ChaosEditor/Source/FractureEditor/Private/FractureToolCustomVoronoi.cpp
@@ -371,7 +371,7 @@ void UFractureToolCustomVoronoi::GenerateLivePattern(const TArray<FFractureToolC
 			// empty ordering can happen if we're in a mode w/ no ImportanceWeights but mode is KeepSharp; just leave default order in that case
 			if (!Ordering.Order.IsEmpty())
 			{
-				TArray<FVector3d> OrderedSites;
+				TArray<FVector> OrderedSites;
 				OrderedSites.SetNumUninitialized(TargetNumSites);
 				for (int32 Idx = 0; Idx < TargetNumSites; Idx++)
 				{
