fix ctrl+alt marquee selection of bones in fracture mode

diff --git a/Engine/Plugins/Experimental/ChaosEditor/Source/FractureEditor/Private/FractureEditorMode.cpp b/Engine/Plugins/Experimental/ChaosEditor/Source/FractureEditor/Private/FractureEditorMode.cpp
--- a/Engine/Plugins/Experimental/ChaosEditor/Source/FractureEditor/Private/FractureEditorMode.cpp
+++ b/Engine/Plugins/Experimental/ChaosEditor/Source/FractureEditor/Private/FractureEditorMode.cpp
@@ -229,9 +229,11 @@ bool UFractureEditorMode::InputAxis(FEditorViewportClient* InViewportClient, FVi
 {
 	// Disable the CTRL+drag behavior of the level editor (translate/rotate the selected object along an axis), since it can easily result in 
 	// unintentionally moving objects while CTRL-clicking to select/deselect
+	// Note: We still want to allow Ctrl+Alt, which is the FrustumSelect in the editor, but not Ctrl+Shift, which is another axis-aligned translation
 
 	bool bCtrlDown = InViewport->KeyState(EKeys::LeftControl) || InViewport->KeyState(EKeys::RightControl);
-	if (bCtrlDown)
+	bool bAltDown = InViewport->KeyState(EKeys::LeftAlt) || InViewport->KeyState(EKeys::RightAlt);
+	if (bCtrlDown && !bAltDown)
 	{
 		return true;
 	}
