Interactive selection tool for fracture mode -- allows a rectangle selection + a filter-by-volume selection

diff --git a/Engine/Plugins/Experimental/ChaosEditor/Source/FractureEditor/Private/FractureEditorCommands.h b/Engine/Plugins/Experimental/ChaosEditor/Source/FractureEditor/Private/FractureEditorCommands.h
--- a/Engine/Plugins/Experimental/ChaosEditor/Source/FractureEditor/Private/FractureEditorCommands.h
+++ b/Engine/Plugins/Experimental/ChaosEditor/Source/FractureEditor/Private/FractureEditorCommands.h
@@ -33,6 +33,7 @@ class FFractureEditorCommands : public TCommands<FFractureEditorCommands>
 		TSharedPtr< FUICommandInfo > SelectSiblings;
 		TSharedPtr< FUICommandInfo > SelectAllInLevel;
 		TSharedPtr< FUICommandInfo > SelectInvert;
+		TSharedPtr< FUICommandInfo > SelectCustom;
 
 		// View Settings
 		TSharedPtr< FUICommandInfo > ToggleShowBoneColors;
diff --git a/Engine/Plugins/Experimental/ChaosEditor/Source/FractureEditor/Private/FractureEditorMode.cpp b/Engine/Plugins/Experimental/ChaosEditor/Source/FractureEditor/Private/FractureEditorMode.cpp
--- a/Engine/Plugins/Experimental/ChaosEditor/Source/FractureEditor/Private/FractureEditorMode.cpp
+++ b/Engine/Plugins/Experimental/ChaosEditor/Source/FractureEditor/Private/FractureEditorMode.cpp
@@ -136,13 +136,23 @@ void UFractureEditorMode::Render(const FSceneView* View, FViewport* Viewport, FP
  
 	if (UFractureModalTool* FractureTool = FractureToolkit->GetActiveTool())
 	{
-		auto Settings = FractureTool->GetSettingsObjects();
 		FractureTool->Render(View, Viewport, PDI);
 	}
 
 
 }
 
+void UFractureEditorMode::DrawHUD(FEditorViewportClient* ViewportClient, FViewport* Viewport, const FSceneView* View, FCanvas* Canvas)
+{
+	Super::DrawHUD(ViewportClient, Viewport, View, Canvas);
+	FFractureEditorModeToolkit* FractureToolkit = (FFractureEditorModeToolkit*)Toolkit.Get();
+
+	if (UFractureModalTool* FractureTool = FractureToolkit->GetActiveTool())
+	{
+		FractureTool->DrawHUD(ViewportClient, Viewport, View, Canvas);
+	}
+}
+
 
 void UFractureEditorMode::CreateToolkit()
 {
@@ -171,17 +181,22 @@ bool UFractureEditorMode::InputKey(FEditorViewportClient* ViewportClient, FViewp
 }
 
 bool UFractureEditorMode::HandleClick(FEditorViewportClient* InViewportClient, HHitProxy* HitProxy, const FViewportClick& Click)
+{
+	return SelectFromClick(HitProxy, Click.IsControlDown(), Click.IsShiftDown());
+}
+
+bool UFractureEditorMode::SelectFromClick(HHitProxy* HitProxy, bool bCtrlDown, bool bShiftDown)
 {
 	if (HitProxy && HitProxy->IsA(HGeometryCollectionBone::StaticGetType()))
 	{
 		HGeometryCollectionBone* GeometryCollectionProxy = (HGeometryCollectionBone*)HitProxy;
 
-		if(GeometryCollectionProxy->Component)
+		if (GeometryCollectionProxy->Component)
 		{
-			TArray<int32> BoneIndices({GeometryCollectionProxy->BoneIndex});
+			TArray<int32> BoneIndices({ GeometryCollectionProxy->BoneIndex });
 
 			GeometryCollectionProxy->Component->Modify();
-			FFractureSelectionTools::ToggleSelectedBones(GeometryCollectionProxy->Component, BoneIndices, !(Click.IsControlDown()||Click.IsShiftDown()), Click.IsShiftDown());			
+			FFractureSelectionTools::ToggleSelectedBones(GeometryCollectionProxy->Component, BoneIndices, !(bCtrlDown || bShiftDown), bShiftDown);
 
 			if (Toolkit.IsValid())
 			{
@@ -206,7 +221,7 @@ bool UFractureEditorMode::HandleClick(FEditorViewportClient* InViewportClient, H
 					TArray<int32> BoneIndices({ TransformIdx });
 
 					OwningGC->Modify();
-					FFractureSelectionTools::ToggleSelectedBones(OwningGC, BoneIndices, !(Click.IsControlDown() || Click.IsShiftDown()), Click.IsShiftDown());
+					FFractureSelectionTools::ToggleSelectedBones(OwningGC, BoneIndices, !(bCtrlDown || bShiftDown), bShiftDown);
 
 					if (Toolkit.IsValid())
 					{
@@ -257,62 +272,124 @@ bool UFractureEditorMode::FrustumSelect(const FConvexVolume& InFrustum, FEditorV
 		for (FSelectionIterator Iter(*SelectedActors); Iter; ++Iter)
 		{
 			AActor* Actor = Cast<AActor>(*Iter);
-			TArray<UGeometryCollectionComponent*, TInlineAllocator<1>> GeometryComponents;
-			Actor->GetComponents<UGeometryCollectionComponent>(GeometryComponents);
+			bSelectedBones = UpdateSelectionInFrustum(InFrustum, Actor, bStrictDragSelection, false, false) || bSelectedBones;
+		}
+	}
+
+	return bSelectedBones;
+}
 
-			if (GeometryComponents.Num())
-			{ 
-				FTransform ActorTransform = Actor->GetTransform();
-				FMatrix InvActorMatrix(ActorTransform.ToInverseMatrixWithScale());
+bool UFractureEditorMode::UpdateSelectionInFrustum(const FConvexVolume& InFrustum, AActor* Actor, bool bStrictDragSelection, bool bAppend, bool bRemove)
+{
+	TArray<UGeometryCollectionComponent*, TInlineAllocator<1>> GeometryComponents;
+	Actor->GetComponents<UGeometryCollectionComponent>(GeometryComponents);
+	if (GeometryComponents.Num() == 0)
+	{
+		return false;
+	}
 
-				FConvexVolume SelectionFrustum(TranformFrustum(InFrustum, InvActorMatrix));
+	FTransform ActorTransform = Actor->GetTransform();
+	FMatrix InvActorMatrix(ActorTransform.ToInverseMatrixWithScale());
 
-				TMap<int32, FBox> BoundsToBone;
-			
-				GetActorGlobalBounds(MakeArrayView(GeometryComponents), BoundsToBone);
+	FConvexVolume SelectionFrustum(TranformFrustum(InFrustum, InvActorMatrix));
 
-				TArray<int32> SelectedBonesArray;
+	TMap<int32, FBox> BoundsToBone;
+	TArray<int32> SelectedBonesArray;
+	bool bSelectionWasUpdated = false;
 
-				for (const auto& Bone : BoundsToBone)
+	for (UGeometryCollectionComponent* GeometryCollectionComponent : GeometryComponents)
+	{
+		BoundsToBone.Reset();
+		SelectedBonesArray.Reset();
+
+		GetComponentGlobalBounds(GeometryCollectionComponent, BoundsToBone);
+
+		for (const TPair<int32, FBox>& Bone : BoundsToBone)
+		{
+			const FBox& TransformedBoneBox = Bone.Value;
+			bool bFullyContained = false;
+			bool bIntersected = SelectionFrustum.IntersectBox(TransformedBoneBox.GetCenter(), TransformedBoneBox.GetExtent(), bFullyContained);
+			if (bIntersected)
+			{
+				if (!bStrictDragSelection || (bFullyContained && bStrictDragSelection))
 				{
-					const FBox& TransformedBoneBox = Bone.Value;
-					bool bFullyContained = false;
-					bool bIntersected = SelectionFrustum.IntersectBox(TransformedBoneBox.GetCenter(), TransformedBoneBox.GetExtent(), bFullyContained);
-					if (bIntersected)
-					{
-						if (!bStrictDragSelection || (bFullyContained && bStrictDragSelection))
-						{
-							SelectedBonesArray.Add(Bone.Key);
-						}
-					}
+					SelectedBonesArray.Add(Bone.Key);
 				}
+			}
+		}
 
-				if (SelectedBonesArray.Num() > 0)
-				{
-					TInlineComponentArray<UGeometryCollectionComponent*> GeometryCollectionComponents;
-					Actor->GetComponents(GeometryCollectionComponents);
+		bool bNeedsUpdate = UpdateSelection(GeometryCollectionComponent->GetSelectedBones(), SelectedBonesArray, bAppend, bRemove);
+		if (bNeedsUpdate)
+		{
+			FScopedColorEdit ColorEdit = GeometryCollectionComponent->EditBoneSelection();
+			ColorEdit.SelectBones(GeometryCollection::ESelectionMode::None);
+			ColorEdit.SetSelectedBones(SelectedBonesArray);
+			ColorEdit.SetHighlightedBones(SelectedBonesArray);
 
-					for (UGeometryCollectionComponent* GeometryCollectionComponent : GeometryCollectionComponents)
-					{
-						FScopedColorEdit ColorEdit = GeometryCollectionComponent->EditBoneSelection();
-						ColorEdit.SelectBones(GeometryCollection::ESelectionMode::None);
-						ColorEdit.SetSelectedBones(SelectedBonesArray);
-						ColorEdit.SetHighlightedBones(SelectedBonesArray);
-						bSelectedBones = true;
+			FFractureEditorModeToolkit* FractureToolkit = (FFractureEditorModeToolkit*)Toolkit.Get();
 
-						FFractureEditorModeToolkit* FractureToolkit = (FFractureEditorModeToolkit*)Toolkit.Get();
+			if (FractureToolkit)
+			{
+				FractureToolkit->SetBoneSelection(GeometryCollectionComponent, ColorEdit.GetSelectedBones(), true);
+			}
 
-						if(FractureToolkit)
-						{
-							FractureToolkit->SetBoneSelection(GeometryCollectionComponent, ColorEdit.GetSelectedBones(), true);
-						}
-					}
-				}
+			bSelectionWasUpdated = true;
+		}
+	}
+
+	return bSelectionWasUpdated;
+}
+
+bool UFractureEditorMode::UpdateSelection(const TArray<int32>& PreviousSelection, TArray<int32>& SelectedBonesArray, bool bAppend, bool bRemove)
+{
+	if (SelectedBonesArray.Num() == 0)
+	{
+		return false;
+	}
+
+	if (!bAppend && !bRemove)
+	{
+		return true;
+	}
+
+	TSet<int32> UpdateSelection(PreviousSelection);
+	if (bAppend)
+	{
+		bool bSelectionChanged = false;
+		for (int32 Bone : SelectedBonesArray)
+		{
+			bool bAlreadyHad = false;
+			UpdateSelection.Add(Bone, &bAlreadyHad);
+			if (bRemove && bAlreadyHad) // bRemove && bAppend == toggle bones
+			{
+				UpdateSelection.Remove(Bone);
+				bSelectionChanged = true;
+			}
+			else if (!bAlreadyHad)
+			{
+				bSelectionChanged = true;
 			}
 		}
+		if (!bSelectionChanged)
+		{
+			return false;
+		}
+	}
+	else if (bRemove)
+	{
+		int32 NumRemoved = 0;
+		for (int32 Bone : SelectedBonesArray)
+		{
+			NumRemoved += UpdateSelection.Remove(Bone);
+		}
+		if (NumRemoved == 0)
+		{
+			return false; // no update actually needed
+		}
 	}
+	SelectedBonesArray = UpdateSelection.Array();
 
-	return bSelectedBones;
+	return true;
 }
 
 bool UFractureEditorMode::ComputeBoundingBoxForViewportFocus(AActor* Actor, UPrimitiveComponent* PrimitiveComponent, FBox& InOutBox) const
@@ -321,7 +398,7 @@ bool UFractureEditorMode::ComputeBoundingBoxForViewportFocus(AActor* Actor, UPri
 	if (GeometryCollectionComponent && GeometryCollectionComponent->GetSelectedBones().Num() > 0 && SelectedGeometryComponents.Contains(GeometryCollectionComponent))
 	{
 		TMap<int32, FBox> BoundsToBoneMap;
-		GetActorGlobalBounds(MakeArrayView<UGeometryCollectionComponent*>(&GeometryCollectionComponent, 1), BoundsToBoneMap);
+		GetComponentGlobalBounds(GeometryCollectionComponent, BoundsToBoneMap);
 
 		FBox TotalBoneBox(ForceInit);
 		for (int32 BoneIndex : GeometryCollectionComponent->GetSelectedBones())
@@ -431,48 +508,45 @@ void UFractureEditorMode::OnActorSelectionChanged(const TArray<UObject*>& NewSel
 	}
 }
 
-void UFractureEditorMode::GetActorGlobalBounds(TArrayView<UGeometryCollectionComponent*> GeometryComponents, TMap<int32, FBox> &BoundsToBone) const
+void UFractureEditorMode::GetComponentGlobalBounds(UGeometryCollectionComponent* GeometryCollectionComponent, TMap<int32, FBox> &BoundsToBone) const
 {
-	for (UGeometryCollectionComponent* GeometryCollectionComponent : GeometryComponents)
-	{
-		FGeometryCollectionEdit RestCollection = GeometryCollectionComponent->EditRestCollection(GeometryCollection::EEditUpdate::None);
-		UGeometryCollection* GeometryCollection = RestCollection.GetRestCollection();
+	FGeometryCollectionEdit RestCollection = GeometryCollectionComponent->EditRestCollection(GeometryCollection::EEditUpdate::None);
+	UGeometryCollection* GeometryCollection = RestCollection.GetRestCollection();
 
-		TSharedPtr<FGeometryCollection, ESPMode::ThreadSafe> GeometryCollectionPtr = GeometryCollection->GetGeometryCollection();
-		FGeometryCollection* OutGeometryCollection = GeometryCollectionPtr.Get();
+	TSharedPtr<FGeometryCollection, ESPMode::ThreadSafe> GeometryCollectionPtr = GeometryCollection->GetGeometryCollection();
+	FGeometryCollection* OutGeometryCollection = GeometryCollectionPtr.Get();
 
-		const TManagedArray<FTransform>& Transform = OutGeometryCollection->GetAttribute<FTransform>("Transform", FGeometryCollection::TransformGroup);
-		const TManagedArray<FBox>& BoundingBox = OutGeometryCollection->GetAttribute<FBox>("BoundingBox", FGeometryCollection::GeometryGroup);
-		const TManagedArray<int32>& TransformToGeometryIndex = OutGeometryCollection->GetAttribute<int32>("TransformToGeometryIndex", FGeometryCollection::TransformGroup);
+	const TManagedArray<FTransform>& Transform = OutGeometryCollection->GetAttribute<FTransform>("Transform", FGeometryCollection::TransformGroup);
+	const TManagedArray<FBox>& BoundingBox = OutGeometryCollection->GetAttribute<FBox>("BoundingBox", FGeometryCollection::GeometryGroup);
+	const TManagedArray<int32>& TransformToGeometryIndex = OutGeometryCollection->GetAttribute<int32>("TransformToGeometryIndex", FGeometryCollection::TransformGroup);
 
-		TManagedArray<FVector3f>* ExplodedVectorsPtr = OutGeometryCollection->FindAttribute<FVector3f>("ExplodedVector", FGeometryCollection::TransformGroup);
+	TManagedArray<FVector3f>* ExplodedVectorsPtr = OutGeometryCollection->FindAttribute<FVector3f>("ExplodedVector", FGeometryCollection::TransformGroup);
 
-		TArray<FTransform> Transforms;
-		GeometryCollectionAlgo::GlobalMatrices(Transform, OutGeometryCollection->Parent, Transforms);
+	TArray<FTransform> Transforms;
+	GeometryCollectionAlgo::GlobalMatrices(Transform, OutGeometryCollection->Parent, Transforms);
 
-		BoundsToBone.Reset();
+	BoundsToBone.Reset();
 
-		if (ExplodedVectorsPtr)
+	if (ExplodedVectorsPtr)
+	{
+		TManagedArray<FVector3f>& ExplodedVectors = *ExplodedVectorsPtr;
+		for (int32 Idx = 0, ni = GeometryCollection->NumElements(FGeometryCollection::TransformGroup); Idx < ni; ++Idx)
 		{
-			TManagedArray<FVector3f>& ExplodedVectors = *ExplodedVectorsPtr;
-			for (int32 Idx = 0, ni = GeometryCollection->NumElements(FGeometryCollection::TransformGroup); Idx < ni; ++Idx)
+			if (TransformToGeometryIndex[Idx] > -1)
 			{
-				if (TransformToGeometryIndex[Idx] > -1)
-				{
-					const FVector3f& Offset = ExplodedVectors[Idx];
-					const FBox& Bounds = BoundingBox[TransformToGeometryIndex[Idx]];
-					BoundsToBone.Add(Idx, Bounds.ShiftBy(Offset).TransformBy(Transforms[Idx]));
-				}
+				const FVector3f& Offset = ExplodedVectors[Idx];
+				const FBox& Bounds = BoundingBox[TransformToGeometryIndex[Idx]];
+				BoundsToBone.Add(Idx, Bounds.ShiftBy(Offset).TransformBy(Transforms[Idx]));
 			}
 		}
-		else
+	}
+	else
+	{
+		for (int32 Idx = 0, ni = GeometryCollection->NumElements(FGeometryCollection::TransformGroup); Idx < ni; ++Idx)
 		{
-			for (int32 Idx = 0, ni = GeometryCollection->NumElements(FGeometryCollection::TransformGroup); Idx < ni; ++Idx)
+			if (TransformToGeometryIndex[Idx] > -1)
 			{
-				if (TransformToGeometryIndex[Idx] > -1)
-				{
-					BoundsToBone.Add(Idx, BoundingBox[TransformToGeometryIndex[Idx]].TransformBy(Transforms[Idx]));
-				}
+				BoundsToBone.Add(Idx, BoundingBox[TransformToGeometryIndex[Idx]].TransformBy(Transforms[Idx]));
 			}
 		}
 	}
diff --git a/Engine/Plugins/Experimental/ChaosEditor/Source/FractureEditor/Private/FractureEditorModeToolkit.cpp b/Engine/Plugins/Experimental/ChaosEditor/Source/FractureEditor/Private/FractureEditorModeToolkit.cpp
--- a/Engine/Plugins/Experimental/ChaosEditor/Source/FractureEditor/Private/FractureEditorModeToolkit.cpp
+++ b/Engine/Plugins/Experimental/ChaosEditor/Source/FractureEditor/Private/FractureEditorModeToolkit.cpp
@@ -763,6 +763,7 @@ void FFractureEditorModeToolkit::BuildToolPalette(FName PaletteIndex, class FToo
 		ToolbarBuilder.AddToolBarButton(Commands.SelectSiblings);
 		ToolbarBuilder.AddToolBarButton(Commands.SelectAllInLevel);
 		ToolbarBuilder.AddToolBarButton(Commands.SelectNeighbors);
+		ToolbarBuilder.AddToolBarButton(Commands.SelectCustom);
 	}
 	else if (PaletteIndex == TEXT("Fracture"))
 	{
diff --git a/Engine/Plugins/Experimental/ChaosEditor/Source/FractureEditor/Private/FractureToolSelection.cpp b/Engine/Plugins/Experimental/ChaosEditor/Source/FractureEditor/Private/FractureToolSelection.cpp
new file mode 100644
--- /dev/null
+++ b/Engine/Plugins/Experimental/ChaosEditor/Source/FractureEditor/Private/FractureToolSelection.cpp
@@ -0,0 +1,492 @@
+// Copyright Epic Games, Inc. All Rights Reserved.
+
+#include "FractureToolSelection.h"
+#include "FractureEditorStyle.h"
+#include "PlanarCut.h"
+#include "GeometryCollection/GeometryCollection.h"
+#include "GeometryCollection/GeometryCollectionObject.h"
+#include "GeometryCollection/GeometryCollectionAlgo.h"
+#include "FractureEditorModeToolkit.h"
+#include "FractureToolContext.h"
+#include "CanvasItem.h"
+#include "CanvasTypes.h"
+#include "EditorModeManager.h"
+#include "FractureEditorMode.h"
+#include "EdModeInteractiveToolsContext.h"
+#include "BaseBehaviors/SingleClickOrDragBehavior.h"
+#include "GameFramework/Actor.h"
+#include "Engine/Selection.h"
+
+#define LOCTEXT_NAMESPACE "FractureSelection"
+
+
+UFractureToolSelection::UFractureToolSelection(const FObjectInitializer& ObjInit) 
+	: Super(ObjInit) 
+{
+	SelectionSettings = NewObject<UFractureSelectionSettings>(GetTransientPackage(), UFractureSelectionSettings::StaticClass());
+	SelectionSettings->OwnerTool = this;
+}
+
+FText UFractureToolSelection::GetDisplayText() const
+{
+	return FText(NSLOCTEXT("FractureSelection", "FractureToolSelection", "Interactive Selection Tool")); 
+}
+
+FText UFractureToolSelection::GetTooltipText() const 
+{
+	return FText(NSLOCTEXT("FractureSelection", "FractureToolSelectionTooltip", "The interactive selection mode provides useful options to help select bones."));
+}
+
+FSlateIcon UFractureToolSelection::GetToolIcon() const 
+{
+	return FSlateIcon("FractureEditorStyle", "FractureEditor.SelectCustom");
+}
+
+void UFractureToolSelection::RegisterUICommand( FFractureEditorCommands* BindingContext ) 
+{
+	UI_COMMAND_EXT( BindingContext, UICommandInfo, "CustomSelection", "Intrctve", "The interactive selection mode provides useful options to help select bones.", EUserInterfaceActionType::ToggleButton, FInputChord() );
+	BindingContext->SelectCustom = UICommandInfo;
+}
+
+void UFractureToolSelection::Setup()
+{
+	// Set things up for being able to add behaviors to live preview.
+	SelectionBehaviorSet = NewObject<UInputBehaviorSet>();
+	SelectionBehaviorSource = NewObject<ULocalInputBehaviorSource>();
+	SelectionBehaviorSource->GetInputBehaviorsFunc = [this]() { return SelectionBehaviorSet; };
+
+	RectangleMarqueeManager = NewObject<URectangleMarqueeManager>();
+	RectangleMarqueeManager->bUseExternalClickDragBehavior = true;
+	RectangleMarqueeManager->Setup([this](UInputBehavior* Behavior, void* Source)
+	{
+		SelectionBehaviorSet->Add(Behavior, Source);
+	});
+	RectangleMarqueeManager->OnDragRectangleStarted.AddUObject(this, &UFractureToolSelection::OnDragRectangleStarted);
+	RectangleMarqueeManager->OnDragRectangleChanged.AddUObject(this, &UFractureToolSelection::OnDragRectangleChanged);
+	RectangleMarqueeManager->OnDragRectangleFinished.AddUObject(this, &UFractureToolSelection::OnDragRectangleFinished);
+
+	USingleClickOrDragInputBehavior* ClickOrDragBehavior = NewObject<USingleClickOrDragInputBehavior>();
+	ClickOrDragBehavior->Initialize(this, RectangleMarqueeManager);
+	ClickOrDragBehavior->Modifiers.RegisterModifier(ShiftModifierID, FInputDeviceState::IsShiftKeyDown);
+	ClickOrDragBehavior->Modifiers.RegisterModifier(CtrlModifierID, FInputDeviceState::IsCtrlKeyDown);
+	SelectionBehaviorSet->Add(ClickOrDragBehavior, this);
+
+	UInteractiveToolsContext* Context = GLevelEditorModeTools().GetInteractiveToolsContext();
+	Context->InputRouter->RegisterSource(SelectionBehaviorSource);
+}
+
+void UFractureToolSelection::OnUpdateModifierState(int ModifierID, bool bIsOn)
+{
+	switch (ModifierID)
+	{
+	case ShiftModifierID:
+		bShiftToggle = bIsOn;
+		break;
+	case CtrlModifierID:
+		bCtrlToggle = bIsOn;
+		break;
+	default:
+		break;
+	}
+}
+
+// These get bound to the delegates on the marquee manager.
+void UFractureToolSelection::OnDragRectangleStarted()
+{
+	// TODO: initialize a live preview of the selection change?
+}
+
+void UFractureToolSelection::OnDragRectangleChanged(const FCameraRectangle& Rectangle)
+{
+	// TODO: implement a live preview of the selection change?
+}
+
+void UFractureToolSelection::OnClicked(const FInputDeviceRay& ClickPos)
+{
+	UInteractiveToolsContext* Context = GLevelEditorModeTools().GetInteractiveToolsContext();
+	UFractureEditorMode* Mode = Cast<UFractureEditorMode>(GLevelEditorModeTools().GetActiveScriptableMode(UFractureEditorMode::EM_FractureEditorModeId));
+	if (!Mode)
+	{
+		return;
+	}
+	FViewport* Viewport = Mode->GetToolManager()->GetContextQueriesAPI()->GetFocusedViewport();
+	if (!Viewport || !Viewport->HasFocus() || !ClickPos.bHas2D)
+	{
+		return;
+	}
+	TRefCountPtr<HHitProxy> HitProxy = Viewport->GetHitProxy(ClickPos.ScreenPosition.X, ClickPos.ScreenPosition.Y);
+	Mode->SelectFromClick(HitProxy, bCtrlToggle, bShiftToggle);
+}
+
+void UFractureToolSelection::OnDragRectangleFinished(const FCameraRectangle& Rectangle, bool bCancelled)
+{
+	if (bCancelled)
+	{
+		return;
+	}
+
+	UFractureEditorMode* FractureMode = Cast<UFractureEditorMode>(GLevelEditorModeTools().GetActiveScriptableMode(UFractureEditorMode::EM_FractureEditorModeId));
+	if (!FractureMode)
+	{
+		return;
+	}
+
+	FConvexVolume FrustumVolume = Rectangle.FrustumAsConvexVolume();
+
+	bool bStrictDragSelection = GetDefault<ULevelEditorViewportSettings>()->bStrictBoxSelection;
+
+	if (USelection* SelectedActors = GEditor->GetSelectedActors())
+	{
+		for (FSelectionIterator Iter(*SelectedActors); Iter; ++Iter)
+		{
+			AActor* Actor = Cast<AActor>(*Iter);
+			FractureMode->UpdateSelectionInFrustum(FrustumVolume, Actor, bStrictDragSelection, bShiftToggle, bCtrlToggle);
+		}
+	}
+}
+
+void URectangleMarqueeManager::Setup(TFunctionRef<void(UInputBehavior*, void*)> AddBehaviorFunc)
+{
+	if (bUseExternalClickDragBehavior == false)
+	{
+		ClickDragBehavior = NewObject<UClickDragInputBehavior>(this);
+		ClickDragBehavior->SetDefaultPriority(BasePriority);
+		ClickDragBehavior->Initialize(this);
+		AddBehaviorFunc(ClickDragBehavior, this);
+	}
+	SetIsEnabled(true);
+}
+
+bool URectangleMarqueeManager::IsEnabled()
+{
+	return bIsEnabled;
+}
+
+void URectangleMarqueeManager::SetIsEnabled(bool bOn)
+{
+	if (bIsDragging && !bOn)
+	{
+		OnTerminateDragSequence();
+	}
+
+	bIsEnabled = bOn;
+}
+
+void URectangleMarqueeManager::SetBasePriority(const FInputCapturePriority& Priority)
+{
+	BasePriority = Priority;
+	if (ClickDragBehavior)
+	{
+		ClickDragBehavior->SetDefaultPriority(Priority);
+	}
+}
+
+FInputCapturePriority URectangleMarqueeManager::GetPriority() const
+{
+	return BasePriority;
+}
+
+void URectangleMarqueeManager::Render(const FSceneView* View, FViewport* Viewport, FPrimitiveDrawInterface* PDI)
+{
+	GLevelEditorModeTools().GetActiveScriptableMode(UFractureEditorMode::EM_FractureEditorModeId)->
+		GetToolManager()->GetContextQueriesAPI()->GetCurrentViewState(CameraRectangle.CameraState);
+}
+
+FInputRayHit URectangleMarqueeManager::CanBeginClickDragSequence(const FInputDeviceRay& PressPos)
+{
+	return bIsEnabled ?
+		FInputRayHit(TNumericLimits<float>::Max()) : // bHit is true. Depth is max to lose the standard tiebreaker.
+		FInputRayHit(); // bHit is false
+}
+
+void URectangleMarqueeManager::OnClickPress(const FInputDeviceRay& PressPos)
+{
+	if (!PressPos.bHas2D)
+	{
+		bIsDragging = false;
+		return;
+	}
+
+	CameraRectangle.RectangleStartRay = PressPos;
+	CameraRectangle.RectangleEndRay = PressPos;
+	CameraRectangle.Initialize();
+
+	bIsOnDragRectangleChangedDeferred = false;
+	OnDragRectangleStarted.Broadcast();
+}
+
+void URectangleMarqueeManager::OnClickDrag(const FInputDeviceRay& DragPos)
+{
+	if (!DragPos.bHas2D)
+	{
+		return;
+	}
+
+	bIsDragging = true;
+	CameraRectangle.RectangleEndRay = DragPos;
+	CameraRectangle.Initialize();
+
+	if (bIsOnDragRectangleChangedDeferred)
+	{
+		return;
+	}
+
+	double Time = 0.0;
+	FDurationTimer Timer(Time);
+	OnDragRectangleChanged.Broadcast(CameraRectangle);
+	Timer.Stop();
+
+	if (Time > OnDragRectangleChangedDeferredThreshold)
+	{
+		bIsOnDragRectangleChangedDeferred = true;
+	}
+}
+
+void URectangleMarqueeManager::OnClickRelease(const FInputDeviceRay& ReleasePos)
+{
+	if (bIsOnDragRectangleChangedDeferred)
+	{
+		bIsOnDragRectangleChangedDeferred = false;
+		OnClickDrag(ReleasePos);
+	}
+
+	bIsDragging = false;
+	OnDragRectangleFinished.Broadcast(CameraRectangle, false);
+}
+
+void URectangleMarqueeManager::OnTerminateDragSequence()
+{
+	bIsDragging = false;
+	OnDragRectangleFinished.Broadcast(CameraRectangle, true);
+}
+
+
+void URectangleMarqueeManager::DrawHUD(FCanvas* Canvas, bool bThisViewHasFocus)
+{
+	if (bThisViewHasFocus && bIsDragging)
+	{
+		FVector2D Start = CameraRectangle.RectangleStartRay.ScreenPosition;
+		FVector2D Curr = CameraRectangle.RectangleEndRay.ScreenPosition;
+		FCanvasBoxItem BoxItem(Start / Canvas->GetDPIScale(), (Curr - Start) / Canvas->GetDPIScale());
+		BoxItem.SetColor(FLinearColor::White);
+		Canvas->DrawItem(BoxItem);
+	}
+}
+
+
+void UFractureToolSelection::Shutdown()
+{
+	UInteractiveToolsContext* Context = GLevelEditorModeTools().GetInteractiveToolsContext();
+	Context->InputRouter->DeregisterSource(SelectionBehaviorSource);
+}
+
+void UFractureToolSelection::DrawHUD(FEditorViewportClient* ViewportClient, FViewport* Viewport, const FSceneView* View, FCanvas* Canvas)
+{
+	Super::DrawHUD(ViewportClient, Viewport, View, Canvas);
+
+	const FViewportClient* Focused = GLevelEditorModeTools().GetFocusedViewportClient();
+	bool bThisViewHasFocus = (ViewportClient == Focused);
+	
+	RectangleMarqueeManager->DrawHUD(Canvas, bThisViewHasFocus);
+}
+
+void UFractureToolSelection::Render(const FSceneView* View, FViewport* Viewport, FPrimitiveDrawInterface* PDI)
+{
+	RectangleMarqueeManager->Render(View, Viewport, PDI);
+
+	EnumerateVisualizationMapping(SelectionMappings, SelectionBounds.Num(), [&](int32 Idx, FVector ExplodedVector)
+	{
+		const FBox& Box = SelectionBounds[Idx];
+		FVector B000 = Box.Min + ExplodedVector;
+		FVector B111 = Box.Max + ExplodedVector;
+		FVector B011(B000.X, B111.Y, B111.Z);
+		FVector B101(B111.X, B000.Y, B111.Z);
+		FVector B110(B111.X, B111.Y, B000.Z);
+		FVector B001(B000.X, B000.Y, B111.Z);
+		FVector B010(B000.X, B111.Y, B000.Z);
+		FVector B100(B111.X, B000.Y, B000.Z);
+		PDI->DrawLine(B000, B100, FLinearColor::Red, SDPG_Foreground, 0.0f, 0.001f);
+		PDI->DrawLine(B000, B010, FLinearColor::Red, SDPG_Foreground, 0.0f, 0.001f);
+		PDI->DrawLine(B000, B001, FLinearColor::Red, SDPG_Foreground, 0.0f, 0.001f);
+		PDI->DrawLine(B111, B011, FLinearColor::Red, SDPG_Foreground, 0.0f, 0.001f);
+		PDI->DrawLine(B111, B101, FLinearColor::Red, SDPG_Foreground, 0.0f, 0.001f);
+		PDI->DrawLine(B111, B110, FLinearColor::Red, SDPG_Foreground, 0.0f, 0.001f);
+		PDI->DrawLine(B001, B101, FLinearColor::Red, SDPG_Foreground, 0.0f, 0.001f);
+		PDI->DrawLine(B001, B011, FLinearColor::Red, SDPG_Foreground, 0.0f, 0.001f);
+		PDI->DrawLine(B110, B100, FLinearColor::Red, SDPG_Foreground, 0.0f, 0.001f);
+		PDI->DrawLine(B110, B010, FLinearColor::Red, SDPG_Foreground, 0.0f, 0.001f);
+		PDI->DrawLine(B100, B101, FLinearColor::Red, SDPG_Foreground, 0.0f, 0.001f);
+		PDI->DrawLine(B010, B011, FLinearColor::Red, SDPG_Foreground, 0.0f, 0.001f);
+	});
+}
+
+TArray<UObject*> UFractureToolSelection::GetSettingsObjects() const
+ {
+	TArray<UObject*> Settings;
+	Settings.Add(SelectionSettings);
+	return Settings;
+}
+
+
+void UFractureToolSelection::FractureContextChanged()
+{
+	RectangleMarqueeManager->SetIsEnabled(SelectionSettings->MouseSelectionMethod == EMouseSelectionMethod::RectSelect);
+
+	UpdateDefaultRandomSeed();
+	TArray<FFractureToolContext> FractureContexts = GetFractureToolContexts();
+
+	ClearVisualizations();
+
+	for (const FFractureToolContext& FractureContext : FractureContexts)
+	{
+		FGeometryCollection& Collection = *FractureContext.GetGeometryCollection();
+		int CollectionIdx = VisualizedCollections.Emplace(FractureContext.GetGeometryCollectionComponent());
+
+		TArray<double> Volumes;
+		FindBoneVolumes(
+			Collection,
+			TArrayView<int32>(), /*Empty array => use all transforms*/
+			Volumes,
+			VolDimScale
+		);
+
+		TInterval<double> VolumeRange = GetVolumeRange(Volumes);
+
+
+		TArray<int32> SelectedIndices;
+		FilterBonesByVolume(
+			Collection,
+			TArrayView<int32>(),
+			Volumes,
+			[&VolumeRange](double Volume, int32 Bone)
+			{
+				return VolumeRange.Contains(Volume);
+			},
+			SelectedIndices
+		);
+
+		FTransform OuterTransform = FractureContext.GetTransform();
+		for (int32 TransformIdx : SelectedIndices)
+		{
+			FTransform InnerTransform = GeometryCollectionAlgo::GlobalMatrix(Collection.Transform, Collection.Parent, TransformIdx);
+
+
+			FTransform CombinedTransform = InnerTransform * OuterTransform;
+
+			int32 GeometryIdx = Collection.TransformToGeometryIndex[TransformIdx];
+			if (GeometryIdx == INDEX_NONE)
+			{
+				ensure(false); // already filtered for existing geo, shouldn't have any missing-geo transforms in the list
+				continue;
+			}
+			int32 VStart = Collection.VertexStart[GeometryIdx];
+			int32 VEnd = VStart + Collection.VertexCount[GeometryIdx];
+			FBox Bounds(EForceInit::ForceInit);
+			for (int32 VIdx = VStart; VIdx < VEnd; VIdx++)
+			{
+				Bounds += CombinedTransform.TransformPosition(Collection.Vertex[VIdx]);
+			}
+			SelectionMappings.AddMapping(CollectionIdx, TransformIdx, SelectionBounds.Num());
+			SelectionBounds.Add(Bounds);
+		}
+	}
+}
+
+TInterval<double> UFractureToolSelection::GetVolumeRange(TArray<double>& Volumes)
+{
+	TInterval<double> Range;
+	if (SelectionSettings->VolumeSelectionMethod == EVolumeSelectionMethod::CubeRoot)
+	{
+		Range.Min = SelectionSettings->MinVolume * VolDimScale;
+		Range.Max = SelectionSettings->MaxVolume * VolDimScale;
+		Range.Min = Range.Min * Range.Min * Range.Min;
+		Range.Max = Range.Max * Range.Max * Range.Max;
+	}
+	else if (SelectionSettings->VolumeSelectionMethod == EVolumeSelectionMethod::RelativeToLargest)
+	{
+		double LargestVolume = KINDA_SMALL_NUMBER;
+		for (double Volume : Volumes)
+		{
+			LargestVolume = FMath::Max(LargestVolume, Volume);
+		}
+		Range.Min = LargestVolume * SelectionSettings->MinVolumeFrac * SelectionSettings->MinVolumeFrac * SelectionSettings->MinVolumeFrac;
+		Range.Max = LargestVolume * SelectionSettings->MaxVolumeFrac * SelectionSettings->MaxVolumeFrac * SelectionSettings->MaxVolumeFrac;
+	}
+	else // EVolumeSelectionMethod::RelativeToWhole
+	{
+		double VolumeSum = 0;
+		for (double Volume : Volumes)
+		{
+			VolumeSum += Volume;
+		}
+		Range.Min = VolumeSum * SelectionSettings->MinVolumeFrac * SelectionSettings->MinVolumeFrac * SelectionSettings->MinVolumeFrac;
+		Range.Max = VolumeSum * SelectionSettings->MaxVolumeFrac * SelectionSettings->MaxVolumeFrac * SelectionSettings->MaxVolumeFrac;
+	}
+	return Range;
+}
+
+
+void UFractureToolSelection::Execute(TWeakPtr<FFractureEditorModeToolkit> InToolkit)
+{
+	if (InToolkit.IsValid())
+	{
+		FFractureEditorModeToolkit* Toolkit = InToolkit.Pin().Get();
+
+		TArray<FFractureToolContext> FractureContexts = GetFractureToolContexts();
+
+		for (FFractureToolContext& FractureContext : FractureContexts)
+		{
+			FractureContext.GetFracturedGeometryCollection()->Modify();
+			FractureContext.GetGeometryCollectionComponent()->Modify();
+
+			UpdateSelection(FractureContext);
+
+			Refresh(FractureContext, Toolkit);
+		}
+
+		SetOutlinerComponents(FractureContexts, Toolkit);
+	}
+}
+
+void UFractureToolSelection::UpdateSelection(FFractureToolContext& FractureContext)
+{
+	if (!FractureContext.GetGeometryCollection().IsValid())
+	{
+		return;
+	}
+
+	FGeometryCollection& Collection = *FractureContext.GetGeometryCollection();
+
+	TArray<double> Volumes;
+	FindBoneVolumes(
+		Collection,
+		TArrayView<int32>(), /*Empty array => use all transforms*/
+		Volumes,
+		VolDimScale
+	);
+
+	TInterval<double> VolumeRange = GetVolumeRange(Volumes);
+
+	TArray<int32> FilterIndices;
+	FilterBonesByVolume(
+		Collection,
+		TArrayView<int32>(),
+		Volumes,
+		[&VolumeRange](double Volume, int32 Bone)
+		{
+			return VolumeRange.Contains(Volume);
+		},
+		FilterIndices
+	);
+
+	UFractureEditorMode* FractureMode = Cast<UFractureEditorMode>(GLevelEditorModeTools().GetActiveScriptableMode(UFractureEditorMode::EM_FractureEditorModeId));
+	if (FractureMode)
+	{
+		bool bAppend = SelectionSettings->SelectionOperation == ESelectionOperation::Add, bRemove = SelectionSettings->SelectionOperation == ESelectionOperation::Remove;
+		FractureMode->UpdateSelection(FractureContext.GetGeometryCollectionComponent()->GetSelectedBones(), FilterIndices, bAppend, bRemove);
+		FractureContext.SetSelection(FilterIndices);
+	}
+}
+
+
+#undef LOCTEXT_NAMESPACE
\ No newline at end of file
diff --git a/Engine/Plugins/Experimental/ChaosEditor/Source/FractureEditor/Private/FractureToolSelection.h b/Engine/Plugins/Experimental/ChaosEditor/Source/FractureEditor/Private/FractureToolSelection.h
new file mode 100644
--- /dev/null
+++ b/Engine/Plugins/Experimental/ChaosEditor/Source/FractureEditor/Private/FractureToolSelection.h
@@ -0,0 +1,275 @@
+// Copyright Epic Games, Inc. All Rights Reserved.
+
+#pragma once
+
+#include "FractureToolCutter.h"
+#include "Engine/StaticMeshActor.h"
+
+
+#include "Mechanics/RectangleMarqueeMechanic.h"
+#include "BaseBehaviors/ClickDragBehavior.h"
+#include "BaseBehaviors/BehaviorTargetInterfaces.h"
+
+#include "FractureToolSelection.generated.h"
+
+class FFractureToolContext;
+
+UENUM()
+enum class EVolumeSelectionMethod
+{
+	// Select by cube root of volume
+	CubeRoot,
+	// Select by cube root of volume relative to the overall shape's cube root of volume
+	RelativeToWhole,
+	// Select by cube root of volume relative to the largest single geometry
+	RelativeToLargest
+};
+
+UENUM()
+enum class ESelectionOperation
+{
+	Replace,
+	Add,
+	Remove
+};
+
+UENUM()
+enum class EMouseSelectionMethod
+{
+	// Click+Drag a rectangle to select -- hold shift to append, ctrl to subtract, and ctrl+shift to toggle
+	RectSelect,
+	// Selection that works the same as when the tool is not active
+	StandardSelect
+	// TODO: lasso selection, brush selection?
+};
+
+/** Settings controlling how geometry is selected */
+UCLASS(config = EditorPerProjectUserSettings)
+class UFractureSelectionSettings : public UFractureToolSettings
+{
+public:
+
+	GENERATED_BODY()
+
+	UFractureSelectionSettings(const FObjectInitializer& ObjInit)
+		: Super(ObjInit)
+	{}
+
+	UPROPERTY(EditAnywhere, Category = MouseInteraction)
+	EMouseSelectionMethod MouseSelectionMethod = EMouseSelectionMethod::RectSelect;
+
+
+
+	/** What values to use when filtering by volume */
+	UPROPERTY(EditAnywhere, Category = FilterSettings)
+	EVolumeSelectionMethod VolumeSelectionMethod = EVolumeSelectionMethod::RelativeToLargest;
+
+	/** How to update the selection from the filter */
+	UPROPERTY(EditAnywhere, Category = FilterSettings)
+	ESelectionOperation SelectionOperation = ESelectionOperation::Replace;
+
+	/** Sets the minimum volume (as computed by the Volume Selection Method) that should be included in the filter */
+	UPROPERTY(EditAnywhere, Category = FilterSettings, meta = (ClampMin = ".00001", UIMin = ".1", UIMax = "1000", EditCondition = "VolumeSelectionMethod == EVolumeSelectionMethod::CubeRoot", EditConditionHides))
+	double MinVolume = 0;
+
+	/** Sets the maximum volume (as computed by the Volume Selection Method) that should be included in the filter */
+	UPROPERTY(EditAnywhere, Category = FilterSettings, meta = (ClampMin = ".00001", UIMin = ".1", UIMax = "1000", EditCondition = "VolumeSelectionMethod == EVolumeSelectionMethod::CubeRoot", EditConditionHides))
+	double MaxVolume = 0;
+
+	/** Sets the minimum volume (as computed by the Volume Selection Method) that should be included in the filter */
+	UPROPERTY(EditAnywhere, Category = FilterSettings, meta = (ClampMin = ".00001", UIMin = ".01", UIMax = "1", EditCondition = "VolumeSelectionMethod != EVolumeSelectionMethod::CubeRoot", EditConditionHides))
+	double MinVolumeFrac = 0;
+
+	/** Sets the maximum volume (as computed by the Volume Selection Method) that should be included in the filter */
+	UPROPERTY(EditAnywhere, Category = FilterSettings, meta = (ClampMin = ".00001", UIMin = ".01", UIMax = "1", EditCondition = "VolumeSelectionMethod != EVolumeSelectionMethod::CubeRoot", EditConditionHides))
+	double MaxVolumeFrac = 0;
+
+};
+
+
+
+
+/*
+ * Adapted from URectangleMarqueeMechanic, the mechanic for a rectangle "marquee" selection, to operate without
+ * the UInteractionMechanic framework. It creates and maintains the 2D rectangle associated with a mouse drag.
+ * It does not test against any scene geometry, nor does it maintain any sort of list of selected objects.
+ *
+ * You must manually call the Setup(), Render() and DrawHUD() functions.
+ * Note these functions have been changed from the Mechanic versions to be easier to call without additional InteractiveTool machinery
+ *
+ * Attach to the delegates and use the passed rectangle to test against your geometry.
+ * 
+ * TODO: Refactor this into a more generalized / generally usable class in RectangleMarqueeMechanic.h,
+ * and share more common code with URectangleMarqueeMechanic.
+ */
+UCLASS()
+class URectangleMarqueeManager : public UObject, public IClickDragBehaviorTarget
+{
+	GENERATED_BODY()
+
+public:
+	/**
+	 * If true, then the URectangleMarqueeManager will not create an internal UClickDragInputBehavior in ::Setup(), allowing
+	 * the client to control the marquee with an external InputBehavior that uses the marquee manager as its IClickDragBehaviorTarget.
+	 * For instance, this allows the mechanic to be hooked in as the drag component of a USingleClickOrDragInputBehavior.
+	 * Must be configured before calling ::Setup()
+	 */
+	UPROPERTY()
+	bool bUseExternalClickDragBehavior = false;
+
+	/**
+	 * If the computation time for a single call to OnDragRectangleChanged ever exceeds this threshold then future calls
+	 * to this function (in the current drag sequence) will be deferred until the mouse button is released. This will
+	 * improve the responsiveness of the UI. The default value is set so this optimization is never triggered, if you
+	 * want it you can set this to a small value e.g., 1./60 (time elapsed by 1 frame at 60 fps)
+	 */
+	UPROPERTY()
+	double OnDragRectangleChangedDeferredThreshold = TNumericLimits<double>::Max(); // In seconds
+
+public:
+
+	void Setup(TFunctionRef<void(UInputBehavior*, void*)> AddBehaviorFunc);
+	void Render(const FSceneView* View, FViewport* Viewport, FPrimitiveDrawInterface* PDI);
+	void DrawHUD(FCanvas* Canvas, bool bThisViewHasFocus);
+
+	bool IsEnabled();
+	void SetIsEnabled(bool bOn);
+
+	/**
+	 * Sets the base priority so that users can make sure that their own behaviors are higher
+	 * priority. The marquee will not use any priority value higher than this.
+	 *
+	 * Can be called before or after Setup(). Only relevant if bUseExternalClickDragBehavior is false
+	 */
+	void SetBasePriority(const FInputCapturePriority& Priority);
+
+	/**
+	 * Gets the current priority used by the marquee behavior
+	 */
+	FInputCapturePriority GetPriority() const;
+
+	/**
+	 * Called when user starts dragging a new rectangle.
+	 */
+	FSimpleMulticastDelegate OnDragRectangleStarted;
+
+	/**
+	 * Called as the user drags the other corner of the rectangle around.
+	 */
+	DECLARE_MULTICAST_DELEGATE_OneParam(OnDragRectangleChangedEvent, const FCameraRectangle&);
+	OnDragRectangleChangedEvent OnDragRectangleChanged;
+
+	/**
+	 * Called once the user lets go of the mouse button after dragging out a rectangle.
+	 * The last dragged rectangle is passed here so that clients can choose to just implement this function in simple cases.
+	 * bCancelled flag is true when the drag finishes due to a disabling of the mechanic or due to a TerminateDragSequence call, rather than a normal drag completion.
+	 */
+	DECLARE_MULTICAST_DELEGATE_TwoParams(OnDragRectangleFinishedEvent, const FCameraRectangle&, bool bCancelled);
+	OnDragRectangleFinishedEvent OnDragRectangleFinished;
+
+protected:
+	UPROPERTY()
+	TObjectPtr<UClickDragInputBehavior> ClickDragBehavior = nullptr;
+
+	FCameraRectangle CameraRectangle;
+
+	FInputCapturePriority BasePriority = FInputCapturePriority(FInputCapturePriority::DEFAULT_TOOL_PRIORITY);
+
+private:
+
+	FInputRayHit CanBeginClickDragSequence(const FInputDeviceRay& PressPos);
+	void OnClickPress(const FInputDeviceRay& PressPos);
+	void OnClickDrag(const FInputDeviceRay& DragPos);
+	void OnClickRelease(const FInputDeviceRay& ReleasePos);
+	void OnTerminateDragSequence();
+
+	bool bIsEnabled;
+	bool bIsDragging;
+	bool bIsOnDragRectangleChangedDeferred;
+};
+
+
+
+
+// Note this tool doesn't actually fracture, but it does remake pieces of geometry and shares a lot of machinery with the fracture tools
+UCLASS(DisplayName = "Custom Selection Tool", Category = "FractureTools")
+class UFractureToolSelection : public UFractureToolCutterBase, public IClickBehaviorTarget
+{
+public:
+	GENERATED_BODY()
+
+	UFractureToolSelection(const FObjectInitializer& ObjInit);
+
+	// UFractureTool Interface
+	virtual FText GetDisplayText() const override;
+	virtual FText GetTooltipText() const override;
+	virtual FSlateIcon GetToolIcon() const override;
+
+	virtual FText GetApplyText() const override { return FText(NSLOCTEXT("FractureSelection", "ExecuteSelection", "Update Selection")); }
+
+	void Render(const FSceneView* View, FViewport* Viewport, FPrimitiveDrawInterface* PDI) override;
+	void DrawHUD(FEditorViewportClient* ViewportClient, FViewport* Viewport, const FSceneView* View, FCanvas* Canvas) override;
+
+	virtual void RegisterUICommand(FFractureEditorCommands* BindingContext) override;
+
+	virtual TArray<UObject*> GetSettingsObjects() const override;
+
+	virtual void FractureContextChanged() override;
+	// ExecuteFracture unused, because this tool doesn't modify geometry; instead, it overrides Execute
+	virtual int32 ExecuteFracture(const FFractureToolContext& FractureContext) override { return INDEX_NONE; }
+	virtual void Execute(TWeakPtr<FFractureEditorModeToolkit> InToolkit) override;
+	void UpdateSelection(FFractureToolContext& FractureContext);
+
+	virtual void Setup() override;
+	virtual void Shutdown() override;
+
+	// IClickBehaviorTarget implementation
+	virtual FInputRayHit IsHitByClick(const FInputDeviceRay& ClickPos) override
+	{
+		return FInputRayHit(FMathf::MaxReal); // Always hit (with MaxReal to lose tiebreakers); actually hit testing done later (in OnClicked)
+	}
+	virtual void OnClicked(const FInputDeviceRay& ClickPos) override;
+	virtual void OnUpdateModifierState(int ModifierID, bool bIsOn) override;
+
+protected:
+	virtual void ClearVisualizations() override
+	{
+		Super::ClearVisualizations();
+		SelectionBounds.Empty();
+		SelectionMappings.Empty();
+	}
+
+	UPROPERTY()
+	TObjectPtr<UInputBehaviorSet> SelectionBehaviorSet = nullptr;
+
+	UPROPERTY()
+	TObjectPtr<ULocalInputBehaviorSource> SelectionBehaviorSource = nullptr;
+
+	UPROPERTY()
+	TObjectPtr<URectangleMarqueeManager> RectangleMarqueeManager = nullptr;
+
+	// Support for Shift and Ctrl toggle
+	bool bShiftToggle = false;
+	bool bCtrlToggle = false;
+	static const int32 ShiftModifierID = 1;
+	static const int32 CtrlModifierID = 2;
+
+	// Rectangle selection
+	void OnDragRectangleStarted();
+	void OnDragRectangleChanged(const FCameraRectangle& Rectangle);
+	void OnDragRectangleFinished(const FCameraRectangle& Rectangle, bool bCancelled);
+
+private:
+
+	UPROPERTY(EditAnywhere, Category = FixGeo)
+	TObjectPtr<UFractureSelectionSettings> SelectionSettings;
+
+	TArray<FBox> SelectionBounds; // Bounds in global space but without exploded vectors applied
+	FVisualizationMappings SelectionMappings;
+
+	TInterval<double> GetVolumeRange(TArray<double>& Volumes);
+	const double VolDimScale = .01; // compute volumes in meters instead of cm, for saner units at typical scales
+
+};
+
+
diff --git a/Engine/Plugins/Experimental/ChaosEditor/Source/FractureEditor/Public/FractureEditorMode.h b/Engine/Plugins/Experimental/ChaosEditor/Source/FractureEditor/Public/FractureEditorMode.h
--- a/Engine/Plugins/Experimental/ChaosEditor/Source/FractureEditor/Public/FractureEditorMode.h
+++ b/Engine/Plugins/Experimental/ChaosEditor/Source/FractureEditor/Public/FractureEditorMode.h
@@ -36,12 +36,16 @@ public:
 
 	//virtual void Tick(FEditorViewportClient* ViewportClient, float DeltaTime) override;
 	virtual void Render(const FSceneView* View, FViewport* Viewport, FPrimitiveDrawInterface* PDI) override;
+	virtual void DrawHUD(FEditorViewportClient* ViewportClient, FViewport* Viewport, const FSceneView* View, FCanvas* Canvas) override;
 
 	virtual void CreateToolkit() override;
 	bool UsesToolkits() const override;
 	virtual bool InputKey(FEditorViewportClient* ViewportClient, FViewport* Viewport, FKey Key, EInputEvent Event) override;
 	virtual bool HandleClick(FEditorViewportClient* InViewportClient, HHitProxy* HitProxy, const FViewportClick& Click) override;
 
+	// Helper for HandleClick, exposed to allow other code paths to handle clicks via the same code path
+	bool SelectFromClick(HHitProxy* HitProxy, bool bCtrlDown, bool bShiftDown);
+
 	virtual bool ComputeBoundingBoxForViewportFocus(AActor* Actor, UPrimitiveComponent* PrimitiveComponent, FBox& InOutBox) const;
 	virtual bool GetPivotForOrbit(FVector& OutPivot) const override;
 	// End of UEdMode interface
@@ -53,6 +57,10 @@ public:
 	virtual bool BoxSelect(FBox& InBox, bool InSelect = true) override;
 	virtual bool FrustumSelect(const FConvexVolume& InFrustum, FEditorViewportClient* InViewportClient, bool InSelect) override;
 
+	// Helpers for FrustumSelect + to expose similar selection functionality to other code
+	bool UpdateSelectionInFrustum(const FConvexVolume& InFrustum, AActor* Actor, bool bStrictDragSelection, bool bAppend, bool bRemove);
+	bool UpdateSelection(const TArray<int32>& PreviousSelection, TArray<int32>& Bones, bool bAppend, bool bRemove);
+
 	// FEditorUndoClient interface
 	virtual bool MatchesContext(const FTransactionContext& InContext, const TArray<TPair<UObject*, FTransactionObjectEvent>>& TransactionObjectContexts) const override;
 	virtual void PostUndo(bool bSuccess) override;
@@ -61,7 +69,7 @@ public:
 private:
 	void OnUndoRedo();
 	void OnActorSelectionChanged(const TArray<UObject*>& NewSelection, bool bForceRefresh);
-	void GetActorGlobalBounds(TArrayView<UGeometryCollectionComponent*> GeometryComponents, TMap<int32, FBox> &BoundsToBone) const;
+	void GetComponentGlobalBounds(UGeometryCollectionComponent* GeometryCollectionComponent, TMap<int32, FBox> &BoundsToBone) const;
 	void SelectionStateChanged();
 	
 	/** Handle package reloading (might be our geometry collection) */
diff --git a/Engine/Plugins/Experimental/ChaosEditor/Source/FractureEditor/Public/FractureTool.h b/Engine/Plugins/Experimental/ChaosEditor/Source/FractureEditor/Public/FractureTool.h
--- a/Engine/Plugins/Experimental/ChaosEditor/Source/FractureEditor/Public/FractureTool.h
+++ b/Engine/Plugins/Experimental/ChaosEditor/Source/FractureEditor/Public/FractureTool.h
@@ -141,6 +141,7 @@ public:
 
 	/** Draw callback from edmode*/
 	virtual void Render(const FSceneView* View, FViewport* Viewport, FPrimitiveDrawInterface* PDI) {}
+	virtual void DrawHUD(FEditorViewportClient* ViewportClient, FViewport* Viewport, const FSceneView* View, FCanvas* Canvas) {}
 
 	virtual void PostEditChangeChainProperty(struct FPropertyChangedChainEvent& PropertyChangedEvent) override;
 	virtual void FractureContextChanged() {}
diff --git a/Engine/Plugins/Experimental/PlanarCutPlugin/Source/PlanarCut/Private/PlanarCut.cpp b/Engine/Plugins/Experimental/PlanarCutPlugin/Source/PlanarCut/Private/PlanarCut.cpp
--- a/Engine/Plugins/Experimental/PlanarCutPlugin/Source/PlanarCut/Private/PlanarCut.cpp
+++ b/Engine/Plugins/Experimental/PlanarCutPlugin/Source/PlanarCut/Private/PlanarCut.cpp
@@ -836,19 +836,19 @@ void FindBoneVolumes(
 }
 
 
-void FindSmallBones(
+void FilterBonesByVolume(
 	FGeometryCollection& Collection,
 	const TArrayView<const int32>& TransformIndices,
 	const TArrayView<const double>& Volumes,
-	double MinVolume,
+	TFunctionRef<bool(double Volume, int32 BoneIdx)> Filter,
 	TArray<int32>& OutSmallBones
 )
 {
 	OutSmallBones.Reset();
 
-	auto AddIdx = [&Collection, &Volumes, &MinVolume, &OutSmallBones](int32 TransformIdx)
+	auto AddIdx = [&Collection, &Volumes, &Filter, &OutSmallBones](int32 TransformIdx)
 	{
-		if (Collection.TransformToGeometryIndex[TransformIdx] > -1 && Volumes[TransformIdx] < MinVolume)
+		if (Collection.TransformToGeometryIndex[TransformIdx] > -1 && Filter(Volumes[TransformIdx], TransformIdx))
 		{
 			OutSmallBones.Add(TransformIdx);
 		}
@@ -881,6 +881,23 @@ void FindSmallBones(
 }
 
 
+void FindSmallBones(
+	FGeometryCollection& Collection,
+	const TArrayView<const int32>& TransformIndices,
+	const TArrayView<const double>& Volumes,
+	double MinVolume,
+	TArray<int32>& OutSmallBones
+)
+{
+	FilterBonesByVolume(Collection, TransformIndices, Volumes, 
+		[MinVolume](double Volume, int32 BoneIdx) -> bool
+		{
+			return Volume < MinVolume;
+		},
+		OutSmallBones);
+}
+
+
 int32 MergeBones(
 	FGeometryCollection& Collection,
 	const TArrayView<const int32>& TransformIndicesView,
diff --git a/Engine/Plugins/Experimental/PlanarCutPlugin/Source/PlanarCut/Public/PlanarCut.h b/Engine/Plugins/Experimental/PlanarCutPlugin/Source/PlanarCut/Public/PlanarCut.h
--- a/Engine/Plugins/Experimental/PlanarCutPlugin/Source/PlanarCut/Public/PlanarCut.h
+++ b/Engine/Plugins/Experimental/PlanarCutPlugin/Source/PlanarCut/Public/PlanarCut.h
@@ -249,6 +249,23 @@ void PLANARCUT_API FindSmallBones(
 	TArray<int32>& OutSmallBones
 );
 
+/**
+ * Populate an array of transform indices w/ those that match a custom volume-based filter
+ *
+ * @param Collection			The collection to be processed
+ * @param TransformIndices		The transform indices to process, or empty if all should be processed
+ * @param Volumes				Volumes of geometry; 1:1 w/ TransformIndices array
+ * @param Filter				Geometry for which the volume filter returns true will be chosen
+ * @param OutSmallBones			Output array, to be filled with transform indices for small pieces of geometry
+ */
+void PLANARCUT_API FilterBonesByVolume(
+	FGeometryCollection& Collection,
+	const TArrayView<const int32>& TransformIndices,
+	const TArrayView<const double>& Volumes,
+	TFunctionRef<bool(double Volume, int32 BoneIdx)> Filter,
+	TArray<int32>& OutSmallBones
+);
+
 namespace UE
 {
 	namespace PlanarCut
diff --git a/Engine/Plugins/Runtime/MeshModelingToolset/Source/ModelingComponents/Private/Mechanics/RectangleMarqueeMechanic.cpp b/Engine/Plugins/Runtime/MeshModelingToolset/Source/ModelingComponents/Private/Mechanics/RectangleMarqueeMechanic.cpp
--- a/Engine/Plugins/Runtime/MeshModelingToolset/Source/ModelingComponents/Private/Mechanics/RectangleMarqueeMechanic.cpp
+++ b/Engine/Plugins/Runtime/MeshModelingToolset/Source/ModelingComponents/Private/Mechanics/RectangleMarqueeMechanic.cpp
@@ -109,6 +109,73 @@ bool FCameraRectangle::IsProjectedSegmentIntersectingRectangle(const FVector& En
 	return TestIntersection(ProjectedSegmentUV, SelectionDomain.Rectangle);
 }
 
+FConvexVolume FCameraRectangle::FrustumAsConvexVolume() const
+{
+	FConvexVolume OutFrustum;
+	ensure(bIsInitialized);
+	FVector2 Point[4]
+	{
+		SelectionDomain.Rectangle.Min,
+		FVector2(SelectionDomain.Rectangle.Min.X, SelectionDomain.Rectangle.Max.Y),
+		SelectionDomain.Rectangle.Max,
+		FVector2(SelectionDomain.Rectangle.Max.X, SelectionDomain.Rectangle.Min.Y)
+	};
+
+	FVector BoxPoint[4];
+	for (int32 i = 0; i < 4; i++)
+	{
+		BoxPoint[i] = PointUVToPoint3D(SelectionDomain.Plane, Point[i]);
+	}
+
+	if (!CameraState.bIsOrthographic)
+	{
+		FVector CamPoint = CameraState.Position;
+		FVector WorldDir[4];
+		for (int32 i = 0; i < 4; i++)
+		{
+			WorldDir[i] = (BoxPoint[i] - CamPoint).GetSafeNormal();
+		}
+		// Use the camera position and the selection box to create the bounding planes
+		FPlane TopPlane(BoxPoint[0], BoxPoint[1], CamPoint); // Top Plane
+		FPlane RightPlane(BoxPoint[1], BoxPoint[2], CamPoint); // Right Plane
+		FPlane BottomPlane(BoxPoint[2], BoxPoint[3], CamPoint); // Bottom Plane
+		FPlane LeftPlane(BoxPoint[3], BoxPoint[0], CamPoint); // Left Plane
+
+		// Near plane is the selection plane flipped, as the convex construction wants the planes to be oriented away from the volume
+		FPlane NearPlane(-SelectionDomain.Plane.Normal, -SelectionDomain.Plane.Constant);
+
+		// TODO: currently no far plane because we haven't propagated that info down to here
+
+		// The frustum is built with the first four planes corresponding to the sides of the frustum + a near plane at the front
+		OutFrustum.Planes.Empty();
+		OutFrustum.Planes.Add(TopPlane);
+		OutFrustum.Planes.Add(RightPlane);
+		OutFrustum.Planes.Add(BottomPlane);
+		OutFrustum.Planes.Add(LeftPlane);
+		OutFrustum.Planes.Add(NearPlane);
+	}
+	else
+	{
+		FVector CamOffset = -SelectionDomain.Plane.Normal;
+		// Use the camera position and the selection box to create the bounding planes
+		FPlane TopPlane(BoxPoint[0], BoxPoint[1], BoxPoint[0] + CamOffset); // Top Plane
+		FPlane RightPlane(BoxPoint[1], BoxPoint[2], BoxPoint[1] + CamOffset); // Right Plane
+		FPlane BottomPlane(BoxPoint[2], BoxPoint[3], BoxPoint[2] + CamOffset); // Bottom Plane
+		FPlane LeftPlane(BoxPoint[3], BoxPoint[0], BoxPoint[3] + CamOffset); // Left Plane
+
+		OutFrustum.Planes.Empty();
+		OutFrustum.Planes.Add(TopPlane);
+		OutFrustum.Planes.Add(RightPlane);
+		OutFrustum.Planes.Add(BottomPlane);
+		OutFrustum.Planes.Add(LeftPlane);
+
+		// Note: Orthographic selections are unbounded in the plane normal direction
+	}
+
+	OutFrustum.Init();
+	return OutFrustum;
+}
+
 // ---------------------------------------
 
 void URectangleMarqueeMechanic::Setup(UInteractiveTool* ParentToolIn)
diff --git a/Engine/Plugins/Runtime/MeshModelingToolset/Source/ModelingComponents/Public/Mechanics/RectangleMarqueeMechanic.h b/Engine/Plugins/Runtime/MeshModelingToolset/Source/ModelingComponents/Public/Mechanics/RectangleMarqueeMechanic.h
--- a/Engine/Plugins/Runtime/MeshModelingToolset/Source/ModelingComponents/Public/Mechanics/RectangleMarqueeMechanic.h
+++ b/Engine/Plugins/Runtime/MeshModelingToolset/Source/ModelingComponents/Public/Mechanics/RectangleMarqueeMechanic.h
@@ -63,6 +63,8 @@ struct MODELINGCOMPONENTS_API FCameraRectangle
 	bool IsProjectedSegmentIntersectingRectangle(const FVector& Endpoint1, const FVector& Endpoint2) const;
 	// bool IsProjectedTriangleIntersectingRectangle(const FVector& A, const FVector& B, const FVector& C) const;
 
+	FConvexVolume FrustumAsConvexVolume() const;
+
 	// Return the 3D point obtained by projecting the given 3D Point onto the given projection plane
 	// Note: Assumes an orthographic camera
 	static FVector OrthographicProjection(const FPlane3& Plane, const FVector& Point)
@@ -71,10 +73,11 @@ struct MODELINGCOMPONENTS_API FCameraRectangle
 	}
 
 	// Return the 3D point obtained by projecting the given 3D Point onto the given projection plane
-	// Note: OrthographicProjection is more efficient if the camera is known to be orthographic
+	// Note: Assumes a perspective camera
 	FVector PerspectiveProjection(const FPlane3& Plane, const FVector& Point) const
 	{
 		ensure(bIsInitialized);
+		checkSlow(!CameraState.bIsOrthographic);
 		return FMath::RayPlaneIntersection(CameraState.Position, Point - CameraState.Position, (FPlane)Plane);
 	}
 
