FractureEditor: Fix Non-Unity compilation issue in UFractureToolSelection::DrawHUD

diff --git a/Engine/Plugins/Experimental/ChaosEditor/Source/FractureEditor/Private/FractureToolSelection.cpp b/Engine/Plugins/Experimental/ChaosEditor/Source/FractureEditor/Private/FractureToolSelection.cpp
--- a/Engine/Plugins/Experimental/ChaosEditor/Source/FractureEditor/Private/FractureToolSelection.cpp
+++ b/Engine/Plugins/Experimental/ChaosEditor/Source/FractureEditor/Private/FractureToolSelection.cpp
@@ -284,8 +284,8 @@ void UFractureToolSelection::DrawHUD(FEditorViewportClient* ViewportClient, FVie
 {
 	Super::DrawHUD(ViewportClient, Viewport, View, Canvas);
 
-	const FViewportClient* Focused = GLevelEditorModeTools().GetFocusedViewportClient();
-	bool bThisViewHasFocus = (ViewportClient == Focused);
+	const FEditorViewportClient* Focused = GLevelEditorModeTools().GetFocusedViewportClient();
+	const bool bThisViewHasFocus = (ViewportClient == Focused);
 	
 	RectangleMarqueeManager->DrawHUD(Canvas, bThisViewHasFocus);
 }
