fix fracture visualization crash when geometry is deleted

#rnx

diff --git a/Engine/Plugins/Experimental/ChaosEditor/Source/FractureEditor/Private/FractureToolCutter.cpp b/Engine/Plugins/Experimental/ChaosEditor/Source/FractureEditor/Private/FractureToolCutter.cpp
--- a/Engine/Plugins/Experimental/ChaosEditor/Source/FractureEditor/Private/FractureToolCutter.cpp
+++ b/Engine/Plugins/Experimental/ChaosEditor/Source/FractureEditor/Private/FractureToolCutter.cpp
@@ -257,6 +257,10 @@ void UFractureToolVoronoiCutterBase::UpdateLineSetExplodedVectors()
 {
 	EnumerateVisualizationMapping(EdgesMappings, VoronoiLineSets.Num(), [&](int32 Idx, FVector ExplodedVector)
 	{
+		if (!VoronoiLineSets[Idx])
+		{
+			return;
+		}
 		// TODO: If we add diagrams even when visibility is toggled off, we also need to update visibility here, e.g.: VoronoiLineSets[Idx]->SetVisibility(CutterSettings->bDrawDiagram);
 		AActor* Actor = VoronoiLineSets[Idx]->GetOwner();
 		VoronoiLineSets[Idx]->SetRelativeLocation(Actor->GetTransform().InverseTransformVector(ExplodedVector));
diff --git a/Engine/Plugins/Experimental/ChaosEditor/Source/FractureEditor/Private/FractureToolCutter.h b/Engine/Plugins/Experimental/ChaosEditor/Source/FractureEditor/Private/FractureToolCutter.h
--- a/Engine/Plugins/Experimental/ChaosEditor/Source/FractureEditor/Private/FractureToolCutter.h
+++ b/Engine/Plugins/Experimental/ChaosEditor/Source/FractureEditor/Private/FractureToolCutter.h
@@ -209,7 +209,10 @@ protected:
 		EdgesMappings.Empty();
 		for (ULineSetComponent* Lines : VoronoiLineSets)
 		{
-			Lines->DestroyComponent();
+			if (Lines)
+			{
+				Lines->DestroyComponent();
+			}
 		}
 		VoronoiLineSets.Empty();
 	}
