Fracture : Convex Fracture
- Raised default to 10.0

diff --git a/Engine/Plugins/Experimental/ChaosEditor/Source/FractureEditor/Private/FractureToolConvex.h b/Engine/Plugins/Experimental/ChaosEditor/Source/FractureEditor/Private/FractureToolConvex.h
--- a/Engine/Plugins/Experimental/ChaosEditor/Source/FractureEditor/Private/FractureToolConvex.h
+++ b/Engine/Plugins/Experimental/ChaosEditor/Source/FractureEditor/Private/FractureToolConvex.h
@@ -26,7 +26,7 @@ public:
 
 	/** We simplify the convex shape to keep points spaced at least this far apart (except to keep the hull from collapsing to zero volume) */
 	UPROPERTY(EditAnywhere, Category = MapSettings, meta = (ClampMin = "0"))
-	double SimplificationDistanceThreshold = 0;
+	double SimplificationDistanceThreshold = 10.0;
 };
 
 
