fix fracture chancetofracture setting (was not hooked up to anything); change meaning to 'per bone' instead of 'per mesh' (jack o says will be more useful)

#rnx

diff --git a/Engine/Plugins/Experimental/ChaosEditor/Source/FractureEditor/Private/FractureToolContext.cpp b/Engine/Plugins/Experimental/ChaosEditor/Source/FractureEditor/Private/FractureToolContext.cpp
--- a/Engine/Plugins/Experimental/ChaosEditor/Source/FractureEditor/Private/FractureToolContext.cpp
+++ b/Engine/Plugins/Experimental/ChaosEditor/Source/FractureEditor/Private/FractureToolContext.cpp
@@ -101,6 +101,18 @@ void FFractureToolContext::ConvertSelectionToLeafNodes(int32 Index, TArray<int32
 	}	
 }
 
+void FFractureToolContext::RandomReduceSelection(float ProbToKeep)
+{
+	FRandomStream RandStream(GetSeed());
+	for (int32 i = 0; i < SelectedBones.Num(); i++)
+	{
+		if (RandStream.GetFraction() >= ProbToKeep) // range does not include 1, so if ProbToKeep is 1 this will never remove
+		{
+			SelectedBones.RemoveAtSwap(i);
+			i--;
+		}
+	}
+}
 
 void FFractureToolContext::ConvertSelectionToRigidNodes()
 {
diff --git a/Engine/Plugins/Experimental/ChaosEditor/Source/FractureEditor/Private/FractureToolContext.h b/Engine/Plugins/Experimental/ChaosEditor/Source/FractureEditor/Private/FractureToolContext.h
--- a/Engine/Plugins/Experimental/ChaosEditor/Source/FractureEditor/Private/FractureToolContext.h
+++ b/Engine/Plugins/Experimental/ChaosEditor/Source/FractureEditor/Private/FractureToolContext.h
@@ -24,6 +24,8 @@ public:
 	void ConvertSelectionToEmbeddedGeometryNodes();
 	void ConvertSelectionToClusterNodes();
 	void RemoveRootNodes();
+	/// @param ProbToKeep The chance to keep each bone in the selection
+	void RandomReduceSelection(float ProbToKeep);
 	
 	const TArray<int32>& GetSelection() const { return SelectedBones; }
 	TArray<int32>& GetSelection() { return SelectedBones; }
diff --git a/Engine/Plugins/Experimental/ChaosEditor/Source/FractureEditor/Private/FractureToolCutter.cpp b/Engine/Plugins/Experimental/ChaosEditor/Source/FractureEditor/Private/FractureToolCutter.cpp
--- a/Engine/Plugins/Experimental/ChaosEditor/Source/FractureEditor/Private/FractureToolCutter.cpp
+++ b/Engine/Plugins/Experimental/ChaosEditor/Source/FractureEditor/Private/FractureToolCutter.cpp
@@ -154,6 +154,7 @@ TArray<FFractureToolContext> UFractureToolCutterBase::GetFractureToolContexts()
 			// Generate a context for each selected node
 			FFractureToolContext FullSelection(GeometryCollectionComponent);
 			FullSelection.ConvertSelectionToRigidNodes();
+			FullSelection.RandomReduceSelection(CutterSettings->ChanceToFracture);
 
 			// Update global transforms and bounds		
 			const TManagedArray<FTransform>& Transform = FullSelection.GetGeometryCollection()->GetAttribute<FTransform>("Transform", FGeometryCollection::TransformGroup);
diff --git a/Engine/Plugins/Experimental/ChaosEditor/Source/FractureEditor/Private/FractureToolCutter.h b/Engine/Plugins/Experimental/ChaosEditor/Source/FractureEditor/Private/FractureToolCutter.h
--- a/Engine/Plugins/Experimental/ChaosEditor/Source/FractureEditor/Private/FractureToolCutter.h
+++ b/Engine/Plugins/Experimental/ChaosEditor/Source/FractureEditor/Private/FractureToolCutter.h
@@ -43,8 +43,8 @@ public:
 	UPROPERTY(EditAnywhere, Category = CommonFracture, meta = (DisplayName = "Random Seed", UIMin = "-1", UIMax = "1000", ClampMin = "-1"))
 	int32 RandomSeed;
 
-	/** Chance to shatter each mesh.  Useful when shattering multiple selected meshes.  */
-	UPROPERTY(EditAnywhere, Category = CommonFracture, meta = (DisplayName = "Chance To Fracture Per Mesh", UIMin = "0.0", UIMax = "1.0", ClampMin = "0.0", ClampMax = "1.0"))
+	/** Chance to shatter each selected bone.  */
+	UPROPERTY(EditAnywhere, Category = CommonFracture, meta = (DisplayName = "Chance To Fracture Per Bone", UIMin = "0.0", UIMax = "1.0", ClampMin = "0.0", ClampMax = "1.0"))
 	float ChanceToFracture;
 
 	/** Generate a fracture pattern across all selected meshes.  */
