try to guard against a reported crash where a geometry collection component has no geometry collection

#rnx

diff --git a/Engine/Plugins/Experimental/ChaosEditor/Source/FractureEditor/Private/FractureEditorMode.cpp b/Engine/Plugins/Experimental/ChaosEditor/Source/FractureEditor/Private/FractureEditorMode.cpp
--- a/Engine/Plugins/Experimental/ChaosEditor/Source/FractureEditor/Private/FractureEditorMode.cpp
+++ b/Engine/Plugins/Experimental/ChaosEditor/Source/FractureEditor/Private/FractureEditorMode.cpp
@@ -463,6 +463,12 @@ void UFractureEditorMode::OnActorSelectionChanged(const TArray<UObject*>& NewSel
 			
 			// If collection does not already have guids, make them
 			FGeometryCollectionEdit RestCollectionEdit = GeometryCollectionComponent->EditRestCollection(GeometryCollection::EEditUpdate::None);
+			if (!ensureMsgf(RestCollectionEdit.GetRestCollection(), TEXT("UGeometryCollectionComponent had no UGeometryCollection")) ||
+				!ensureMsgf(RestCollectionEdit.GetRestCollection()->GetGeometryCollection(), TEXT("UGeometryCollectionComponent had no FGeometryCollection")))
+			{
+				// guard against the component not having a UGeometryCollection or FGeometryCollection, with ensures because it doesn't seem like this should happen
+				continue;
+			}
 			::GeometryCollection::GenerateTemporaryGuids(RestCollectionEdit.GetRestCollection()->GetGeometryCollection().Get());
 			
 			FScopedColorEdit ShowBoneColorsEdit(GeometryCollectionComponent);
