preserve materials across geometry collection resets, as much as possible

#rnx

diff --git a/Engine/Plugins/Experimental/ChaosEditor/Source/FractureEditor/Private/FractureToolGenerators.cpp b/Engine/Plugins/Experimental/ChaosEditor/Source/FractureEditor/Private/FractureToolGenerators.cpp
--- a/Engine/Plugins/Experimental/ChaosEditor/Source/FractureEditor/Private/FractureToolGenerators.cpp
+++ b/Engine/Plugins/Experimental/ChaosEditor/Source/FractureEditor/Private/FractureToolGenerators.cpp
@@ -389,10 +389,16 @@ void UFractureToolResetAsset::Execute(TWeakPtr<FFractureEditorModeToolkit> InToo
 		FGeometryCollectionEdit GeometryCollectionEdit = GeometryCollectionComponent->EditRestCollection();
 		if (UGeometryCollection* GeometryCollectionObject = GeometryCollectionEdit.GetRestCollection())
 		{
-
 			TSharedPtr<FGeometryCollection, ESPMode::ThreadSafe> GeometryCollectionPtr = GeometryCollectionObject->GetGeometryCollection();
 			if (FGeometryCollection* GeometryCollection = GeometryCollectionPtr.Get())
 			{
+				constexpr bool bKeepPreviousMaterials = true; // written as a flag in case we want to make this optional later
+				TArray<TObjectPtr<UMaterialInterface>> OldMaterials;
+				if (bKeepPreviousMaterials)
+				{
+					OldMaterials = GeometryCollectionObject->Materials;
+				}
+
 				GeometryCollectionObject->Reset();
 
 				// Rebuild Collection from recorded source assets.
@@ -416,6 +422,23 @@ void UFractureToolResetAsset::Execute(TWeakPtr<FFractureEditorModeToolkit> InToo
 
 				GeometryCollectionObject->InitializeMaterials();
 
+				if (bKeepPreviousMaterials)
+				{
+					int32 NewMatNum = GeometryCollectionObject->Materials.Num(), OldMatNum = OldMaterials.Num();
+					// if the source asset was changed, number of materials might have changed; only copy to the extent the two arrays match
+					int32 NumToCopy = FMath::Min(NewMatNum, OldMatNum);
+					for (int32 MatIdx = 0; MatIdx + 1 < NumToCopy; MatIdx++)
+					{
+						GeometryCollectionObject->Materials[MatIdx] = OldMaterials[MatIdx];
+					}
+					if (NumToCopy > 0) // copy the selection material
+					{
+						GeometryCollectionObject->Materials[NewMatNum - 1] = OldMaterials[OldMatNum - 1];
+					}
+				}
+				
+
+
 				FGeometryCollectionClusteringUtility::UpdateHierarchyLevelOfChildren(GeometryCollection, -1);
 				AddSingleRootNodeIfRequired(GeometryCollectionObject);
 				GeometryCollectionComponent->MarkRenderDynamicDataDirty();
