fix fracture mode generate-new-collection not correctly remembering long asset paths

#rnx

diff --git a/Engine/Plugins/Experimental/ChaosEditor/Source/FractureEditor/Private/FractureToolGenerators.cpp b/Engine/Plugins/Experimental/ChaosEditor/Source/FractureEditor/Private/FractureToolGenerators.cpp
--- a/Engine/Plugins/Experimental/ChaosEditor/Source/FractureEditor/Private/FractureToolGenerators.cpp
+++ b/Engine/Plugins/Experimental/ChaosEditor/Source/FractureEditor/Private/FractureToolGenerators.cpp
@@ -108,7 +108,7 @@ void UFractureToolGenerateAsset::OnGenerateAssetPathChosen(const FString& InAsse
 	int32 LastSlash = INDEX_NONE;
 	if (InAssetPath.FindLastChar('/', LastSlash))
 	{
-		AssetPath = InAssetPath.LeftChop(LastSlash);
+		AssetPath = InAssetPath.Left(LastSlash);
 	}
 
 	UGeometryCollectionComponent* GeometryCollectionComponent = nullptr;
