add icon for fracture mode's interactive selection tool

#rnx

diff --git a/Engine/Plugins/Experimental/ChaosEditor/Source/FractureEditor/Private/FractureEditorStyle.cpp b/Engine/Plugins/Experimental/ChaosEditor/Source/FractureEditor/Private/FractureEditorStyle.cpp
--- a/Engine/Plugins/Experimental/ChaosEditor/Source/FractureEditor/Private/FractureEditorStyle.cpp
+++ b/Engine/Plugins/Experimental/ChaosEditor/Source/FractureEditor/Private/FractureEditorStyle.cpp
@@ -39,6 +39,7 @@ FFractureEditorStyle::FFractureEditorStyle()
 	Set("FractureEditor.SelectSiblings",       new IMAGE_BRUSH_SVG("FractureSelectSiblings", IconSize));
 	Set("FractureEditor.SelectAllInLevel",     new IMAGE_BRUSH_SVG("FractureSelectLevel", IconSize));
 	Set("FractureEditor.SelectInvert",         new IMAGE_BRUSH_SVG("FractureSelectInvert", IconSize));
+	Set("FractureEditor.SelectCustom",         new IMAGE_BRUSH_SVG("FractureSelectInteractive", IconSize));
 
 	Set("FractureEditor.AutoCluster",          new IMAGE_BRUSH_SVG("FractureAutoCluster", IconSize));
 	Set("FractureEditor.ClusterMagnet",        new IMAGE_BRUSH_SVG("FractureMagnet", IconSize));
diff --git a/Engine/Plugins/Experimental/ChaosEditor/Source/FractureEditor/Private/FractureToolSelection.cpp b/Engine/Plugins/Experimental/ChaosEditor/Source/FractureEditor/Private/FractureToolSelection.cpp
--- a/Engine/Plugins/Experimental/ChaosEditor/Source/FractureEditor/Private/FractureToolSelection.cpp
+++ b/Engine/Plugins/Experimental/ChaosEditor/Source/FractureEditor/Private/FractureToolSelection.cpp
@@ -44,7 +44,7 @@ FSlateIcon UFractureToolSelection::GetToolIcon() const
 
 void UFractureToolSelection::RegisterUICommand( FFractureEditorCommands* BindingContext ) 
 {
-	UI_COMMAND_EXT( BindingContext, UICommandInfo, "CustomSelection", "Intrctve", "The interactive selection mode provides useful options to help select bones.", EUserInterfaceActionType::ToggleButton, FInputChord() );
+	UI_COMMAND_EXT( BindingContext, UICommandInfo, "SelectCustom", "Intrctve", "The interactive selection mode provides useful options to help select bones.", EUserInterfaceActionType::ToggleButton, FInputChord() );
 	BindingContext->SelectCustom = UICommandInfo;
 }
 
