combine fracture tool categories "uv" "fix" and "properties" into a new single "utilities" category

#rnx

diff --git a/Engine/Plugins/Experimental/ChaosEditor/Source/FractureEditor/Private/FractureEditorModeToolkit.cpp b/Engine/Plugins/Experimental/ChaosEditor/Source/FractureEditor/Private/FractureEditorModeToolkit.cpp
--- a/Engine/Plugins/Experimental/ChaosEditor/Source/FractureEditor/Private/FractureEditorModeToolkit.cpp
+++ b/Engine/Plugins/Experimental/ChaosEditor/Source/FractureEditor/Private/FractureEditorModeToolkit.cpp
@@ -671,7 +671,7 @@ void FFractureEditorModeToolkit::OnObjectPostEditChange( UObject* Object, FPrope
 	}
 }
 
-const TArray<FName> FFractureEditorModeToolkit::PaletteNames = { FName(TEXT("Generate")), FName(TEXT("Select")), FName(TEXT("Fracture")), FName(TEXT("Edit")), FName(TEXT("Cluster")), FName(TEXT("Embed")), FName(TEXT("Fix")), FName(TEXT("UV")), FName(TEXT("Properties")) };
+const TArray<FName> FFractureEditorModeToolkit::PaletteNames = { FName(TEXT("Generate")), FName(TEXT("Select")), FName(TEXT("Fracture")), FName(TEXT("Edit")), FName(TEXT("Cluster")), FName(TEXT("Embed")), FName(TEXT("Utilities")) };
 
 FText FFractureEditorModeToolkit::GetToolPaletteDisplayName(FName Palette) const
 { 
@@ -798,21 +798,15 @@ void FFractureEditorModeToolkit::BuildToolPalette(FName PaletteIndex, class FToo
 		ToolbarBuilder.AddToolBarButton(Commands.AutoEmbedGeometry);
 		ToolbarBuilder.AddToolBarButton(Commands.FlushEmbeddedGeometry);
 	}
-	else if (PaletteIndex == TEXT("UV"))
+	else if (PaletteIndex == TEXT("Utilities"))
 	{
 		ToolbarBuilder.AddToolBarButton(Commands.AutoUV);
-	}
-	else if (PaletteIndex == TEXT("Fix"))
-	{
 		ToolbarBuilder.AddToolBarButton(Commands.RecomputeNormals);
 		ToolbarBuilder.AddToolBarButton(Commands.Resample);
 		ToolbarBuilder.AddToolBarButton(Commands.ConvertToMesh);
 		ToolbarBuilder.AddToolBarButton(Commands.Validate);
 		ToolbarBuilder.AddToolBarButton(Commands.MakeConvex);
 		ToolbarBuilder.AddToolBarButton(Commands.FixTinyGeo);
-	}
-	else if (PaletteIndex == TEXT("Properties"))
-	{
 		ToolbarBuilder.AddToolBarButton(Commands.SetInitialDynamicState);
 	}
 }
