stop automatically destroying / recreating physics state with the fracture tool context constructor via misuse of FGeometryCollectionEdit, and instead do the needed updates on fracture tool execution

#rnx

diff --git a/Engine/Plugins/Experimental/ChaosEditor/Source/FractureEditor/Private/FractureToolContext.cpp b/Engine/Plugins/Experimental/ChaosEditor/Source/FractureEditor/Private/FractureToolContext.cpp
--- a/Engine/Plugins/Experimental/ChaosEditor/Source/FractureEditor/Private/FractureToolContext.cpp
+++ b/Engine/Plugins/Experimental/ChaosEditor/Source/FractureEditor/Private/FractureToolContext.cpp
@@ -14,8 +14,7 @@ FFractureToolContext::FFractureToolContext(UGeometryCollectionComponent* InGeome
 
 	Transform = GeometryCollectionComponent->GetOwner()->GetActorTransform();
 
-	FGeometryCollectionEdit GeometryCollectionEdit = GeometryCollectionComponent->EditRestCollection();
-	if (UGeometryCollection* GeometryCollectionObject = GeometryCollectionEdit.GetRestCollection())
+	if (UGeometryCollection* GeometryCollectionObject = RestCollection.GetRestCollection())
 	{
 		GeometryCollection = GeometryCollectionObject->GetGeometryCollection();
 		SelectedBones = GeometryCollectionComponent->GetSelectedBones();
diff --git a/Engine/Plugins/Experimental/ChaosEditor/Source/FractureEditor/Private/FractureToolRecomputeNormals.h b/Engine/Plugins/Experimental/ChaosEditor/Source/FractureEditor/Private/FractureToolRecomputeNormals.h
--- a/Engine/Plugins/Experimental/ChaosEditor/Source/FractureEditor/Private/FractureToolRecomputeNormals.h
+++ b/Engine/Plugins/Experimental/ChaosEditor/Source/FractureEditor/Private/FractureToolRecomputeNormals.h
@@ -75,6 +75,11 @@ public:
 
 	virtual void FractureContextChanged() override;
 	virtual int32 ExecuteFracture(const FFractureToolContext& FractureContext) override;
+	// don't need to update physics state for new normals/tangents
+	virtual bool ExecuteUpdatesPhysics() const override
+	{
+		return false;
+	}
 
 protected:
 	UPROPERTY()
diff --git a/Engine/Plugins/Experimental/ChaosEditor/Source/FractureEditor/Private/FractureToolUV.h b/Engine/Plugins/Experimental/ChaosEditor/Source/FractureEditor/Private/FractureToolUV.h
--- a/Engine/Plugins/Experimental/ChaosEditor/Source/FractureEditor/Private/FractureToolUV.h
+++ b/Engine/Plugins/Experimental/ChaosEditor/Source/FractureEditor/Private/FractureToolUV.h
@@ -208,6 +208,11 @@ public:
 	/** Executes function that generates new geometry. Returns the first new geometry index. */
 	virtual int32 ExecuteFracture(const FFractureToolContext& FractureContext) override;
 	virtual bool CanExecute() const override;
+	// don't need to update physics state when the tool just changes UVs
+	virtual bool ExecuteUpdatesPhysics() const override
+	{
+		return false;
+	}
 	virtual void PostEditChangeChainProperty(struct FPropertyChangedChainEvent& PropertyChangedEvent) override;
 
 	/** Gets the UI command info for this command */
diff --git a/Engine/Plugins/Experimental/ChaosEditor/Source/FractureEditor/Public/FractureTool.cpp b/Engine/Plugins/Experimental/ChaosEditor/Source/FractureEditor/Public/FractureTool.cpp
--- a/Engine/Plugins/Experimental/ChaosEditor/Source/FractureEditor/Public/FractureTool.cpp
+++ b/Engine/Plugins/Experimental/ChaosEditor/Source/FractureEditor/Public/FractureTool.cpp
@@ -49,16 +49,27 @@ const TSharedPtr<FUICommandInfo>& UFractureActionTool::GetUICommandInfo() const
 	return UICommandInfo;
 }
 
-UFractureActionTool::FModifyContextScope::FModifyContextScope(UFractureActionTool* ActionTool, FFractureToolContext* FractureContext) : ActionTool(ActionTool), FractureContext(FractureContext)
+UFractureActionTool::FModifyContextScope::FModifyContextScope(UFractureActionTool* ActionTool, FFractureToolContext* FractureContext, bool bWantPhysicsUpdate) : ActionTool(ActionTool), FractureContext(FractureContext)
 {
 	check(FractureContext);
 	check(ActionTool);
 	FractureContext->GetFracturedGeometryCollection()->Modify();
 	FractureContext->GetGeometryCollectionComponent()->Modify();
+	bool bHasPhysicsState = FractureContext->GetGeometryCollectionComponent()->HasValidPhysicsState();
+	bNeedPhysicsUpdate = bWantPhysicsUpdate && bHasPhysicsState;
+	if (bNeedPhysicsUpdate)
+	{
+		FractureContext->GetGeometryCollectionComponent()->DestroyPhysicsState();
+	}
 }
 
 UFractureActionTool::FModifyContextScope::~FModifyContextScope()
 {
+	if (bNeedPhysicsUpdate)
+	{
+		FractureContext->GetGeometryCollectionComponent()->RecreatePhysicsState();
+	}
+
 	UFractureEditorMode* FractureMode = Cast<UFractureEditorMode>(GLevelEditorModeTools().GetActiveScriptableMode(UFractureEditorMode::EM_FractureEditorModeId));
 	if (!FractureMode)
 	{
@@ -298,6 +309,12 @@ void UFractureModalTool::Execute(TWeakPtr<FFractureEditorModeToolkit> InToolkit)
 		{
 			FractureContext.GetFracturedGeometryCollection()->Modify();
 			FractureContext.GetGeometryCollectionComponent()->Modify();
+			bool bHadPhysicsState = FractureContext.GetGeometryCollectionComponent()->HasValidPhysicsState();
+			bool bNeedPhysicsUpdate = ExecuteUpdatesPhysics() && bHadPhysicsState;
+			if (bNeedPhysicsUpdate)
+			{
+				FractureContext.GetGeometryCollectionComponent()->DestroyPhysicsState();
+			}
 
 			int32 FirstNewGeometryIndex = ExecuteFracture(FractureContext);
 			
@@ -339,6 +356,11 @@ void UFractureModalTool::Execute(TWeakPtr<FFractureEditorModeToolkit> InToolkit)
 				FractureContext.GetFracturedGeometryCollection()->InitResources();
 			}
 
+			if (bNeedPhysicsUpdate)
+			{
+				FractureContext.GetGeometryCollectionComponent()->RecreatePhysicsState();
+			}
+
 			Refresh(FractureContext, Toolkit);
 		}
 
diff --git a/Engine/Plugins/Experimental/ChaosEditor/Source/FractureEditor/Public/FractureTool.h b/Engine/Plugins/Experimental/ChaosEditor/Source/FractureEditor/Public/FractureTool.h
--- a/Engine/Plugins/Experimental/ChaosEditor/Source/FractureEditor/Public/FractureTool.h
+++ b/Engine/Plugins/Experimental/ChaosEditor/Source/FractureEditor/Public/FractureTool.h
@@ -70,8 +70,9 @@ public:
 	{
 		UFractureActionTool* ActionTool;
 		FFractureToolContext* FractureContext;
+		bool bNeedPhysicsUpdate;
 
-		FModifyContextScope(UFractureActionTool* ActionTool, FFractureToolContext* FractureContext);
+		FModifyContextScope(UFractureActionTool* ActionTool, FFractureToolContext* FractureContext, bool bWantPhysicsUpdate = true);
 
 		~FModifyContextScope();
 	};
@@ -147,6 +148,11 @@ public:
 	/** Executes the command.  Derived types need to be implemented in a thread safe way*/
 	virtual void Execute(TWeakPtr<FFractureEditorModeToolkit> InToolkit) override;
 	virtual bool CanExecute() const override;
+	// @return true if the physics state needs to be updated after executing (e.g., can be false if tool is just updating UVs)
+	virtual bool ExecuteUpdatesPhysics() const
+	{
+		return true;
+	}
 
 	/** Executes function that generates new geometry. Returns the first new geometry index. */
 	virtual int32 ExecuteFracture(const FFractureToolContext& FractureContext) { return INDEX_NONE; }
