remove convex hull promotion from the fracture mode UI, because having multiple convex hulls on a single transform is not supported by sim

#rnx

diff --git a/Engine/Plugins/Experimental/ChaosEditor/Source/FractureEditor/Private/FractureToolConvex.cpp b/Engine/Plugins/Experimental/ChaosEditor/Source/FractureEditor/Private/FractureToolConvex.cpp
--- a/Engine/Plugins/Experimental/ChaosEditor/Source/FractureEditor/Private/FractureToolConvex.cpp
+++ b/Engine/Plugins/Experimental/ChaosEditor/Source/FractureEditor/Private/FractureToolConvex.cpp
@@ -19,11 +19,14 @@ void UFractureConvexSettings::DeleteFromSelected()
 	ConvexTool->DeleteConvexFromSelected();
 }
 
+// Note: this feature puts multiple convexes on a single bone, which isn't supported by sim yet
+#if 0
 void UFractureConvexSettings::PromoteChildren()
 {
 	UFractureToolConvex* ConvexTool = Cast<UFractureToolConvex>(OwnerTool.Get());
 	ConvexTool->PromoteChildren();
 }
+#endif
 
 void UFractureConvexSettings::ClearCustomConvex()
 {
diff --git a/Engine/Plugins/Experimental/ChaosEditor/Source/FractureEditor/Private/FractureToolConvex.h b/Engine/Plugins/Experimental/ChaosEditor/Source/FractureEditor/Private/FractureToolConvex.h
--- a/Engine/Plugins/Experimental/ChaosEditor/Source/FractureEditor/Private/FractureToolConvex.h
+++ b/Engine/Plugins/Experimental/ChaosEditor/Source/FractureEditor/Private/FractureToolConvex.h
@@ -36,9 +36,12 @@ public:
 	UFUNCTION(CallInEditor, Category = Custom, meta = (DisplayName = "Delete From Selected"))
 	void DeleteFromSelected();
 
+	// Note: this feature puts multiple convexes on a single bone, which isn't supported by sim yet
+#if 0
 	/** Promote (and save) child convex hulls on to the selected bone(s) */
 	UFUNCTION(CallInEditor, Category = Custom, meta = (DisplayName = "Promote Children"))
 	void PromoteChildren();
+#endif
 
 	/** Clear any manual adjustments to convex hulls on the selected bones */
 	UFUNCTION(CallInEditor, Category = Custom, meta = (DisplayName = "Clear Custom Convex"))
