add modify() calls to fracture autouv tool actions so they show up in the undo stack

#rnx

diff --git a/Engine/Plugins/Experimental/ChaosEditor/Source/FractureEditor/Private/FractureToolUV.cpp b/Engine/Plugins/Experimental/ChaosEditor/Source/FractureEditor/Private/FractureToolUV.cpp
--- a/Engine/Plugins/Experimental/ChaosEditor/Source/FractureEditor/Private/FractureToolUV.cpp
+++ b/Engine/Plugins/Experimental/ChaosEditor/Source/FractureEditor/Private/FractureToolUV.cpp
@@ -141,6 +141,8 @@ void UFractureToolAutoUV::UpdateUVChannels(int32 TargetNumUVChannels)
 
 	if (TargetNumUVChannels > -1 && TargetNumUVChannels != MinUVChannels)
 	{
+		FScopedTransaction Transaction(LOCTEXT("UpdateUVChannels", "Update UV Channels"), !GeomCompSelection.IsEmpty());
+
 		bool bIsIncreasing = TargetNumUVChannels > MinUVChannels;
 		for (UGeometryCollectionComponent* GeometryCollectionComponent : GeomCompSelection)
 		{
@@ -149,6 +151,8 @@ void UFractureToolAutoUV::UpdateUVChannels(int32 TargetNumUVChannels)
 			// but if we're increasing the number of channels, some geometry collections may have more channels already and we leave those alone
 			if (NumChannels < TargetNumUVChannels || !bIsIncreasing)
 			{
+				FGeometryCollectionEdit Edit(GeometryCollectionComponent, GeometryCollection::EEditUpdate::None);
+				Edit.GetRestCollection()->Modify();
 				GeometryCollectionComponent->GetRestCollection()->GetGeometryCollection()->SetNumUVLayers(TargetNumUVChannels);
 				if (bIsIncreasing)
 				{
@@ -162,6 +166,8 @@ void UFractureToolAutoUV::UpdateUVChannels(int32 TargetNumUVChannels)
 						}
 					}
 				}
+				GeometryCollectionComponent->MarkRenderDynamicDataDirty();
+				GeometryCollectionComponent->MarkRenderStateDirty();
 			}
 		}
 		MinUVChannels = TargetNumUVChannels;
@@ -181,8 +187,12 @@ void UFractureToolAutoUV::BoxProjectUVs()
 
 	TSet<UGeometryCollectionComponent*> GeomCompSelection;
 	GetSelectedGeometryCollectionComponents(GeomCompSelection);
+
+	FScopedTransaction Transaction(LOCTEXT("BoxProjectUVs", "Box Project UVs"), !GeomCompSelection.IsEmpty());
 	for (UGeometryCollectionComponent* GeometryCollectionComponent : GeomCompSelection)
 	{
+		FGeometryCollectionEdit Edit(GeometryCollectionComponent, GeometryCollection::EEditUpdate::None);
+		Edit.GetRestCollection()->Modify();
 		UE::PlanarCut::BoxProjectUVs(UVLayer,
 			*GeometryCollectionComponent->GetRestCollection()->GetGeometryCollection(),
 			(FVector3d)AutoUVSettings->ProjectionScale,
@@ -370,8 +380,12 @@ void UFractureToolAutoUV::LayoutUVs()
 {
 	TSet<UGeometryCollectionComponent*> GeomCompSelection;
 	GetSelectedGeometryCollectionComponents(GeomCompSelection);
+
+	FScopedTransaction Transaction(LOCTEXT("LayoutUVs", "Layout UVs"), !GeomCompSelection.IsEmpty());
 	for (UGeometryCollectionComponent* GeometryCollectionComponent : GeomCompSelection)
 	{
+		FGeometryCollectionEdit Edit(GeometryCollectionComponent, GeometryCollection::EEditUpdate::None);
+		Edit.GetRestCollection()->Modify();
 		LayoutUVsForComponent(GeometryCollectionComponent);
 
 		GeometryCollectionComponent->MarkRenderDynamicDataDirty();
@@ -402,12 +416,11 @@ void UFractureToolAutoUV::BakeTexture()
 {
 	TSet<UGeometryCollectionComponent*> GeomCompSelection;
 	GetSelectedGeometryCollectionComponents(GeomCompSelection);
+
+	FScopedTransaction Transaction(LOCTEXT("BakeTexture", "Bake Texture"), !GeomCompSelection.IsEmpty());
 	for (UGeometryCollectionComponent* GeometryCollectionComponent : GeomCompSelection)
 	{
 		BakeTextureForComponent(GeometryCollectionComponent);
-
-		GeometryCollectionComponent->MarkRenderDynamicDataDirty();
-		GeometryCollectionComponent->MarkRenderStateDirty();
 	}
 }
 
