workaround for fracture mode outliner UI not existing when mode Enter() tries to refresh the UI with the current selection

#rnx

diff --git a/Engine/Plugins/Experimental/ChaosEditor/Source/FractureEditor/Private/FractureEditorMode.cpp b/Engine/Plugins/Experimental/ChaosEditor/Source/FractureEditor/Private/FractureEditorMode.cpp
--- a/Engine/Plugins/Experimental/ChaosEditor/Source/FractureEditor/Private/FractureEditorMode.cpp
+++ b/Engine/Plugins/Experimental/ChaosEditor/Source/FractureEditor/Private/FractureEditorMode.cpp
@@ -514,6 +514,15 @@ void UFractureEditorMode::OnActorSelectionChanged(const TArray<UObject*>& NewSel
 	}
 }
 
+void UFractureEditorMode::RefreshOutlinerWithCurrentSelection()
+{
+	if (Toolkit.IsValid())
+	{
+		FFractureEditorModeToolkit* FractureToolkit = (FFractureEditorModeToolkit*)Toolkit.Get();
+		FractureToolkit->SetOutlinerComponents(SelectedGeometryComponents);
+	}
+}
+
 void UFractureEditorMode::GetComponentGlobalBounds(UGeometryCollectionComponent* GeometryCollectionComponent, TMap<int32, FBox> &BoundsToBone) const
 {
 	FGeometryCollectionEdit RestCollection = GeometryCollectionComponent->EditRestCollection(GeometryCollection::EEditUpdate::None);
diff --git a/Engine/Plugins/Experimental/ChaosEditor/Source/FractureEditor/Private/FractureEditorModeToolkit.cpp b/Engine/Plugins/Experimental/ChaosEditor/Source/FractureEditor/Private/FractureEditorModeToolkit.cpp
--- a/Engine/Plugins/Experimental/ChaosEditor/Source/FractureEditor/Private/FractureEditorModeToolkit.cpp
+++ b/Engine/Plugins/Experimental/ChaosEditor/Source/FractureEditor/Private/FractureEditorModeToolkit.cpp
@@ -511,6 +511,13 @@ void FFractureEditorModeToolkit::InvokeUI()
 			}
 		}
 	}
+
+	// TODO: This is a workaround for the Mode's Enter() call happening *before* the toolkit UI-building call.  Ideally we'd find a better way of making sure the toolkit UI knows about the current selection.
+	UFractureEditorMode* Mode = Cast<UFractureEditorMode>(GLevelEditorModeTools().GetActiveScriptableMode(UFractureEditorMode::EM_FractureEditorModeId));
+	if (Mode)
+	{
+		Mode->RefreshOutlinerWithCurrentSelection();
+	}
 }
 
 TSharedRef<SDockTab> FFractureEditorModeToolkit::CreateHierarchyTab(const FSpawnTabArgs& Args)
diff --git a/Engine/Plugins/Experimental/ChaosEditor/Source/FractureEditor/Public/FractureEditorMode.h b/Engine/Plugins/Experimental/ChaosEditor/Source/FractureEditor/Public/FractureEditorMode.h
--- a/Engine/Plugins/Experimental/ChaosEditor/Source/FractureEditor/Public/FractureEditorMode.h
+++ b/Engine/Plugins/Experimental/ChaosEditor/Source/FractureEditor/Public/FractureEditorMode.h
@@ -66,6 +66,10 @@ public:
 	virtual void PostUndo(bool bSuccess) override;
 	virtual void PostRedo(bool bSuccess) override;
 
+	// Used to update outliner when the selection hasn't change -- needed because the outliner is not built yet when Enter() is called
+	// TODO: We may be able to remove this function if we can find a way to invoke Enter's current selection-update logic after the outliner UI is built
+	void RefreshOutlinerWithCurrentSelection();
+
 private:
 	void OnUndoRedo();
 	void OnActorSelectionChanged(const TArray<UObject*>& NewSelection, bool bForceRefresh);
