skip the last material on the geometry collection when converting to static mesh, as long as it is unused -- by convention this is the selection material that is only used for highlighting parts of the geometry collection, and we don't want it to make it back to the static mesh

#rnx

diff --git a/Engine/Plugins/Experimental/ChaosEditor/Source/FractureEditor/Private/FractureToolConvert.cpp b/Engine/Plugins/Experimental/ChaosEditor/Source/FractureEditor/Private/FractureToolConvert.cpp
--- a/Engine/Plugins/Experimental/ChaosEditor/Source/FractureEditor/Private/FractureToolConvert.cpp
+++ b/Engine/Plugins/Experimental/ChaosEditor/Source/FractureEditor/Private/FractureToolConvert.cpp
@@ -309,9 +309,26 @@ int32 UFractureToolConvert::ExecuteFracture(const FFractureToolContext& Fracture
 
 		TArray<TObjectPtr<class UMaterialInterface>>& OverrideMaterials = FractureContext.GetGeometryCollectionComponent()->OverrideMaterials;
 		TArray<TObjectPtr<class UMaterialInterface>>& MainMaterials = FractureContext.GetFracturedGeometryCollection()->Materials;
+
+		int32 LastMaterialIdx = MainMaterials.Num() - 1;
+		// if possible, we'd like to skip the last material -- it should be the selection material, which should be unused
+		int32 SkipLastMaterialOffset = 1;
+		if (LastMaterialIdx >= 0)
+		{
+			for (int32 UsedMaterialID : Collection.MaterialID)
+			{
+				if (UsedMaterialID == LastMaterialIdx)
+				{
+					// last material is used in the mesh; we can't skip it
+					SkipLastMaterialOffset = 0;
+					break;
+				}
+			}
+		}
+
 		TArray<UMaterialInterface*> Materials;
 		Materials.Reserve(MainMaterials.Num());
-		for (int MatIdx = 0; MatIdx < MainMaterials.Num(); MatIdx++)
+		for (int32 MatIdx = 0; MatIdx + SkipLastMaterialOffset < MainMaterials.Num(); MatIdx++)
 		{
 			if (MatIdx < OverrideMaterials.Num() && !OverrideMaterials[MatIdx].IsNull())
 			{
