use FGeometryCollectionEdit more consistently and make sure convex hulls are updated where needed (and not updated when not needed)

#rnx

diff --git a/Engine/Plugins/Experimental/ChaosEditor/Source/FractureEditor/Private/FractureToolAutoCluster.cpp b/Engine/Plugins/Experimental/ChaosEditor/Source/FractureEditor/Private/FractureToolAutoCluster.cpp
--- a/Engine/Plugins/Experimental/ChaosEditor/Source/FractureEditor/Private/FractureToolAutoCluster.cpp
+++ b/Engine/Plugins/Experimental/ChaosEditor/Source/FractureEditor/Private/FractureToolAutoCluster.cpp
@@ -69,6 +69,8 @@ void UFractureToolAutoCluster::Execute(TWeakPtr<FFractureEditorModeToolkit> InTo
 		
 		for (FFractureToolContext& Context : Contexts)
 		{
+			FGeometryCollectionEdit Edit(Context.GetGeometryCollectionComponent(), GeometryCollection::EEditUpdate::RestPhysicsDynamic);
+
 			Context.ConvertSelectionToClusterNodes();
 			FGeometryCollection* GeometryCollection = Context.GetGeometryCollection().Get();
 
diff --git a/Engine/Plugins/Experimental/ChaosEditor/Source/FractureEditor/Private/FractureToolClusterMagnet.cpp b/Engine/Plugins/Experimental/ChaosEditor/Source/FractureEditor/Private/FractureToolClusterMagnet.cpp
--- a/Engine/Plugins/Experimental/ChaosEditor/Source/FractureEditor/Private/FractureToolClusterMagnet.cpp
+++ b/Engine/Plugins/Experimental/ChaosEditor/Source/FractureEditor/Private/FractureToolClusterMagnet.cpp
@@ -84,6 +84,8 @@ void UFractureToolClusterMagnet::Execute(TWeakPtr<FFractureEditorModeToolkit> In
 				continue;
 			}
 
+			FGeometryCollectionEdit Edit(Context.GetGeometryCollectionComponent(), GeometryCollection::EEditUpdate::RestPhysicsDynamic);
+
 			const TManagedArray<TSet<int32>>& Children = Context.GetGeometryCollection()->Children;
 			const TManagedArray<int32>& Levels = Context.GetGeometryCollection()->GetAttribute<int32>("Level", FGeometryCollection::TransformGroup);
 			int32 StartTransformCount = Children.Num();
@@ -186,12 +188,11 @@ TMap<int32, TSet<int32>> UFractureToolClusterMagnet::InitializeConnectivity(cons
 
 void UFractureToolClusterMagnet::CollectTopNodeConnections(FGeometryCollectionPtr GeometryCollection, int32 Index, int32 OperatingLevel, TSet<int32>& OutConnections) const
 {
-	const TManagedArray<TSet<int32>>& Children = GeometryCollection->Children;
-	if (Children[Index].Num() == 0) // leaf node
+	const TManagedArray<int32>& TransformToGeometryIndex = GeometryCollection->TransformToGeometryIndex;
+	if (TransformToGeometryIndex[Index] == INDEX_NONE) // leaf node
 	{
 		const TManagedArray<TSet<int32>>& Proximity = GeometryCollection->GetAttribute<TSet<int32>>("Proximity", FGeometryCollection::GeometryGroup);
 		const TManagedArray<int32>& GeometryToTransformIndex = GeometryCollection->TransformIndex;
-		const TManagedArray<int32>& TransformToGeometryIndex = GeometryCollection->TransformToGeometryIndex;
 
 
 		for (int32 Neighbor : Proximity[TransformToGeometryIndex[Index]])
@@ -202,6 +203,7 @@ void UFractureToolClusterMagnet::CollectTopNodeConnections(FGeometryCollectionPt
 	}
 	else
 	{
+		const TManagedArray<TSet<int32>>& Children = GeometryCollection->Children;
 		for (int32 ChildIndex : Children[Index])
 		{
 			CollectTopNodeConnections(GeometryCollection, ChildIndex, OperatingLevel, OutConnections);
diff --git a/Engine/Plugins/Experimental/ChaosEditor/Source/FractureEditor/Private/FractureToolClustering.cpp b/Engine/Plugins/Experimental/ChaosEditor/Source/FractureEditor/Private/FractureToolClustering.cpp
--- a/Engine/Plugins/Experimental/ChaosEditor/Source/FractureEditor/Private/FractureToolClustering.cpp
+++ b/Engine/Plugins/Experimental/ChaosEditor/Source/FractureEditor/Private/FractureToolClustering.cpp
@@ -44,6 +44,8 @@ void UFractureToolFlattenAll::Execute(TWeakPtr<FFractureEditorModeToolkit> InToo
 
 		for (FFractureToolContext& Context : Contexts)
 		{
+			FGeometryCollectionEdit Edit(Context.GetGeometryCollectionComponent(), GeometryCollection::EEditUpdate::RestPhysicsDynamic);
+
 			const TManagedArray<int32>& Levels = Context.GetGeometryCollection()->GetAttribute<int32>("Level", FGeometryCollection::TransformGroup);
 
 			Context.ConvertSelectionToClusterNodes();
@@ -103,6 +105,8 @@ void UFractureToolCluster::Execute(TWeakPtr<FFractureEditorModeToolkit> InToolki
 
 		for (FFractureToolContext& Context : Contexts)
 		{
+			FGeometryCollectionEdit Edit(Context.GetGeometryCollectionComponent(), GeometryCollection::EEditUpdate::RestPhysicsDynamic);
+
 			int32 StartTransformCount = Context.GetGeometryCollection()->Transform.Num();
 			
 			Context.RemoveRootNodes();
@@ -173,6 +177,8 @@ void UFractureToolUncluster::Execute(TWeakPtr<FFractureEditorModeToolkit> InTool
 
 		for (FFractureToolContext& Context : Contexts)
 		{
+			FGeometryCollectionEdit Edit(Context.GetGeometryCollectionComponent(), GeometryCollection::EEditUpdate::RestPhysicsDynamic);
+
 			const TManagedArray<TSet<int32>>& Children = Context.GetGeometryCollection()->GetAttribute<TSet<int32>>("Children", FGeometryCollection::TransformGroup);
 			const TManagedArray<int32>& Levels = Context.GetGeometryCollection()->GetAttribute<int32>("Level", FGeometryCollection::TransformGroup);
 
@@ -224,6 +230,8 @@ void UFractureToolMoveUp::Execute(TWeakPtr<FFractureEditorModeToolkit> InToolkit
 
 		for (FFractureToolContext& Context: Contexts)
 		{
+			FGeometryCollectionEdit Edit(Context.GetGeometryCollectionComponent(), GeometryCollection::EEditUpdate::RestPhysicsDynamic);
+
 			Context.ConvertSelectionToRigidNodes();
 			FGeometryCollectionClusteringUtility::MoveUpOneHierarchyLevel(Context.GetGeometryCollection().Get(), Context.GetSelection());
 			FGeometryCollectionClusteringUtility::RemoveDanglingClusters(Context.GetGeometryCollection().Get());
@@ -269,6 +277,8 @@ void UFractureToolClusterMerge::Execute(TWeakPtr<FFractureEditorModeToolkit> InT
 
 		for (FFractureToolContext& Context : Contexts)
 		{
+			FGeometryCollectionEdit Edit(Context.GetGeometryCollectionComponent(), GeometryCollection::EEditUpdate::RestPhysicsDynamic);
+
 			const TManagedArray<TSet<int32>>& Children = Context.GetGeometryCollection()->Children;
 
 			Context.ConvertSelectionToClusterNodes();
diff --git a/Engine/Plugins/Experimental/ChaosEditor/Source/FractureEditor/Private/FractureToolConvert.h b/Engine/Plugins/Experimental/ChaosEditor/Source/FractureEditor/Private/FractureToolConvert.h
--- a/Engine/Plugins/Experimental/ChaosEditor/Source/FractureEditor/Private/FractureToolConvert.h
+++ b/Engine/Plugins/Experimental/ChaosEditor/Source/FractureEditor/Private/FractureToolConvert.h
@@ -60,6 +60,10 @@ public:
 	/** Executes function that generates new geometry. Returns the first new geometry index. */
 	virtual int32 ExecuteFracture(const FFractureToolContext& FractureContext) override;
 	virtual bool CanExecute() const override;
+	virtual bool ExecuteUpdatesShape() const override
+	{
+		return false;
+	}
 	virtual void PostEditChangeChainProperty(struct FPropertyChangedChainEvent& PropertyChangedEvent) override;
 
 	/** Gets the UI command info for this command */
diff --git a/Engine/Plugins/Experimental/ChaosEditor/Source/FractureEditor/Private/FractureToolEditing.cpp b/Engine/Plugins/Experimental/ChaosEditor/Source/FractureEditor/Private/FractureToolEditing.cpp
--- a/Engine/Plugins/Experimental/ChaosEditor/Source/FractureEditor/Private/FractureToolEditing.cpp
+++ b/Engine/Plugins/Experimental/ChaosEditor/Source/FractureEditor/Private/FractureToolEditing.cpp
@@ -46,6 +46,7 @@ void UFractureToolDeleteBranch::Execute(TWeakPtr<FFractureEditorModeToolkit> InT
 
 		for (FFractureToolContext& Context : Contexts)
 		{
+			FGeometryCollectionEdit Edit(Context.GetGeometryCollectionComponent(), GeometryCollection::EEditUpdate::RestPhysicsDynamic);
 			FGeometryCollection* GeometryCollection = Context.GetGeometryCollection().Get();
 			UGeometryCollection* FracturedGeometryCollection = Context.GetFracturedGeometryCollection();
 
diff --git a/Engine/Plugins/Experimental/ChaosEditor/Source/FractureEditor/Private/FractureToolEmbed.cpp b/Engine/Plugins/Experimental/ChaosEditor/Source/FractureEditor/Private/FractureToolEmbed.cpp
--- a/Engine/Plugins/Experimental/ChaosEditor/Source/FractureEditor/Private/FractureToolEmbed.cpp
+++ b/Engine/Plugins/Experimental/ChaosEditor/Source/FractureEditor/Private/FractureToolEmbed.cpp
@@ -60,7 +60,7 @@ void UFractureToolAddEmbeddedGeometry::Execute(TWeakPtr<FFractureEditorModeToolk
 
 		for (FFractureToolContext& Context : Contexts)
 		{
-			FGeometryCollectionEdit GeometryCollectionEdit = Context.GetGeometryCollectionComponent()->EditRestCollection(GeometryCollection::EEditUpdate::RestPhysicsDynamic);
+			FGeometryCollectionEdit GeometryCollectionEdit = Context.GetGeometryCollectionComponent()->EditRestCollection(GeometryCollection::EEditUpdate::RestPhysicsDynamic, true /*bShapeIsUnchanged*/);
 			UGeometryCollection* FracturedGeometryCollection = GeometryCollectionEdit.GetRestCollection();
 
 			FFractureToolContext::FGeometryCollectionPtr GeometryCollection = Context.GetGeometryCollection();
@@ -188,6 +188,8 @@ void UFractureToolAutoEmbedGeometry::Execute(TWeakPtr<FFractureEditorModeToolkit
 
 			for (FFractureToolContext& Context : Contexts)
 			{
+				FGeometryCollectionEdit Edit(Context.GetGeometryCollectionComponent(), GeometryCollection::EEditUpdate::RestPhysicsDynamic, true /*bShapeIsUnchanged*/);
+
 				FGeometryCollection* GeometryCollection = Context.GetGeometryCollection().Get();
 				FGeometryCollectionConvexUtility::FGeometryCollectionConvexData ConvexData = FGeometryCollectionConvexUtility::GetValidConvexHullData(GeometryCollection);
 				const TManagedArray<FTransform>& Transform = GeometryCollection->Transform;
@@ -247,7 +249,7 @@ void UFractureToolAutoEmbedGeometry::Execute(TWeakPtr<FFractureEditorModeToolkit
 				if (BoneIndex > INDEX_NONE)
 				{
 					// We found the closest bone, now we embed the geometry.
-					FGeometryCollectionEdit GeometryCollectionEdit = ClosestComponent->EditRestCollection(GeometryCollection::EEditUpdate::RestPhysicsDynamic);
+					FGeometryCollectionEdit GeometryCollectionEdit = ClosestComponent->EditRestCollection(GeometryCollection::EEditUpdate::RestPhysicsDynamic, true /*bShapeIsUnchanged*/);
 					UGeometryCollection* FracturedGeometryCollection = GeometryCollectionEdit.GetRestCollection();
 
 					FTransform BoneGlobalTransform = GeometryCollectionAlgo::GlobalMatrix(Transform, Parent, BoneIndex);
@@ -398,7 +400,7 @@ void UFractureToolFlushEmbeddedGeometry::Execute(TWeakPtr<FFractureEditorModeToo
 			if (RemovalList.Num())
 			{
 				// Remove all instanced embedded geometry.
-				FGeometryCollectionEdit GeometryCollectionEdit = Context.GetGeometryCollectionComponent()->EditRestCollection(GeometryCollection::EEditUpdate::RestPhysicsDynamic);
+				FGeometryCollectionEdit GeometryCollectionEdit = Context.GetGeometryCollectionComponent()->EditRestCollection(GeometryCollection::EEditUpdate::RestPhysicsDynamic, true /*bShapeIsUnchanged*/);
 				UGeometryCollection* FracturedCollection = GeometryCollectionEdit.GetRestCollection();
 				FracturedCollection->GetGeometryCollection()->RemoveElements(FGeometryCollection::TransformGroup, RemovalList);
 
diff --git a/Engine/Plugins/Experimental/ChaosEditor/Source/FractureEditor/Private/FractureToolProperties.cpp b/Engine/Plugins/Experimental/ChaosEditor/Source/FractureEditor/Private/FractureToolProperties.cpp
--- a/Engine/Plugins/Experimental/ChaosEditor/Source/FractureEditor/Private/FractureToolProperties.cpp
+++ b/Engine/Plugins/Experimental/ChaosEditor/Source/FractureEditor/Private/FractureToolProperties.cpp
@@ -59,7 +59,7 @@ void UFractureToolSetInitialDynamicState::Execute(TWeakPtr<FFractureEditorModeTo
 		GetSelectedGeometryCollectionComponents(GeomCompSelection);
 		for (UGeometryCollectionComponent* GeometryCollectionComponent : GeomCompSelection)
 		{
-			FGeometryCollectionEdit GCEdit = GeometryCollectionComponent->EditRestCollection();
+			FGeometryCollectionEdit GCEdit = GeometryCollectionComponent->EditRestCollection(GeometryCollection::EEditUpdate::RestPhysics, true /*bShapeIsUnchanged*/);
 			if (UGeometryCollection* GCObject = GCEdit.GetRestCollection())
 			{
 				TSharedPtr<FGeometryCollection, ESPMode::ThreadSafe> GeometryCollectionPtr = GCObject->GetGeometryCollection();
diff --git a/Engine/Plugins/Experimental/ChaosEditor/Source/FractureEditor/Private/FractureToolRecomputeNormals.h b/Engine/Plugins/Experimental/ChaosEditor/Source/FractureEditor/Private/FractureToolRecomputeNormals.h
--- a/Engine/Plugins/Experimental/ChaosEditor/Source/FractureEditor/Private/FractureToolRecomputeNormals.h
+++ b/Engine/Plugins/Experimental/ChaosEditor/Source/FractureEditor/Private/FractureToolRecomputeNormals.h
@@ -75,8 +75,7 @@ public:
 
 	virtual void FractureContextChanged() override;
 	virtual int32 ExecuteFracture(const FFractureToolContext& FractureContext) override;
-	// don't need to update physics state for new normals/tangents
-	virtual bool ExecuteUpdatesPhysics() const override
+	virtual bool ExecuteUpdatesShape() const override
 	{
 		return false;
 	}
diff --git a/Engine/Plugins/Experimental/ChaosEditor/Source/FractureEditor/Private/FractureToolResample.h b/Engine/Plugins/Experimental/ChaosEditor/Source/FractureEditor/Private/FractureToolResample.h
--- a/Engine/Plugins/Experimental/ChaosEditor/Source/FractureEditor/Private/FractureToolResample.h
+++ b/Engine/Plugins/Experimental/ChaosEditor/Source/FractureEditor/Private/FractureToolResample.h
@@ -51,6 +51,10 @@ public:
 
 	virtual void FractureContextChanged() override;
 	virtual int32 ExecuteFracture(const FFractureToolContext& FractureContext) override;
+	virtual bool ExecuteUpdatesShape() const override
+	{
+		return false;
+	}
 
 protected:
 	virtual void ClearVisualizations() override
diff --git a/Engine/Plugins/Experimental/ChaosEditor/Source/FractureEditor/Private/FractureToolUV.h b/Engine/Plugins/Experimental/ChaosEditor/Source/FractureEditor/Private/FractureToolUV.h
--- a/Engine/Plugins/Experimental/ChaosEditor/Source/FractureEditor/Private/FractureToolUV.h
+++ b/Engine/Plugins/Experimental/ChaosEditor/Source/FractureEditor/Private/FractureToolUV.h
@@ -208,8 +208,7 @@ public:
 	/** Executes function that generates new geometry. Returns the first new geometry index. */
 	virtual int32 ExecuteFracture(const FFractureToolContext& FractureContext) override;
 	virtual bool CanExecute() const override;
-	// don't need to update physics state when the tool just changes UVs
-	virtual bool ExecuteUpdatesPhysics() const override
+	virtual bool ExecuteUpdatesShape() const override
 	{
 		return false;
 	}
diff --git a/Engine/Plugins/Experimental/ChaosEditor/Source/FractureEditor/Public/FractureTool.cpp b/Engine/Plugins/Experimental/ChaosEditor/Source/FractureEditor/Public/FractureTool.cpp
--- a/Engine/Plugins/Experimental/ChaosEditor/Source/FractureEditor/Public/FractureTool.cpp
+++ b/Engine/Plugins/Experimental/ChaosEditor/Source/FractureEditor/Public/FractureTool.cpp
@@ -307,14 +307,7 @@ void UFractureModalTool::Execute(TWeakPtr<FFractureEditorModeToolkit> InToolkit)
 
 		for (FFractureToolContext& FractureContext : FractureContexts)
 		{
-			FractureContext.GetFracturedGeometryCollection()->Modify();
-			FractureContext.GetGeometryCollectionComponent()->Modify();
-			bool bHadPhysicsState = FractureContext.GetGeometryCollectionComponent()->HasValidPhysicsState();
-			bool bNeedPhysicsUpdate = ExecuteUpdatesPhysics() && bHadPhysicsState;
-			if (bNeedPhysicsUpdate)
-			{
-				FractureContext.GetGeometryCollectionComponent()->DestroyPhysicsState();
-			}
+			FGeometryCollectionEdit EditCollection(FractureContext.GetGeometryCollectionComponent(), GeometryCollection::EEditUpdate::RestPhysicsDynamic, !ExecuteUpdatesShape());
 
 			int32 FirstNewGeometryIndex = ExecuteFracture(FractureContext);
 			
@@ -356,13 +349,6 @@ void UFractureModalTool::Execute(TWeakPtr<FFractureEditorModeToolkit> InToolkit)
 				FractureContext.GetFracturedGeometryCollection()->InitResources();
 			}
 
-			FractureContext.GetGeometryCollectionComponent()->EditRestCollection(GeometryCollection::EEditUpdate::RestPhysicsDynamic);
-
-			if (bNeedPhysicsUpdate)
-			{
-				FractureContext.GetGeometryCollectionComponent()->RecreatePhysicsState();
-			}
-
 			Refresh(FractureContext, Toolkit);
 		}
 
diff --git a/Engine/Plugins/Experimental/ChaosEditor/Source/FractureEditor/Public/FractureTool.h b/Engine/Plugins/Experimental/ChaosEditor/Source/FractureEditor/Public/FractureTool.h
--- a/Engine/Plugins/Experimental/ChaosEditor/Source/FractureEditor/Public/FractureTool.h
+++ b/Engine/Plugins/Experimental/ChaosEditor/Source/FractureEditor/Public/FractureTool.h
@@ -148,8 +148,8 @@ public:
 	/** Executes the command.  Derived types need to be implemented in a thread safe way*/
 	virtual void Execute(TWeakPtr<FFractureEditorModeToolkit> InToolkit) override;
 	virtual bool CanExecute() const override;
-	// @return true if the physics state needs to be updated after executing (e.g., can be false if tool is just updating UVs)
-	virtual bool ExecuteUpdatesPhysics() const
+	// @return true if the edit will change the shape data of the geometry collection(s) (including changing the clusters) -- used to indicate whether we need to update the convex hulls
+	virtual bool ExecuteUpdatesShape() const
 	{
 		return true;
 	}
diff --git a/Engine/Source/Runtime/Experimental/GeometryCollectionEngine/Private/GeometryCollection/GeometryCollectionComponent.cpp b/Engine/Source/Runtime/Experimental/GeometryCollectionEngine/Private/GeometryCollection/GeometryCollectionComponent.cpp
--- a/Engine/Source/Runtime/Experimental/GeometryCollectionEngine/Private/GeometryCollection/GeometryCollectionComponent.cpp
+++ b/Engine/Source/Runtime/Experimental/GeometryCollectionEngine/Private/GeometryCollection/GeometryCollectionComponent.cpp
@@ -2229,9 +2229,10 @@ void UGeometryCollectionComponent::SetRestCollection(const UGeometryCollection*
 	}
 }
 
-FGeometryCollectionEdit::FGeometryCollectionEdit(UGeometryCollectionComponent* InComponent, GeometryCollection::EEditUpdate InEditUpdate)
+FGeometryCollectionEdit::FGeometryCollectionEdit(UGeometryCollectionComponent* InComponent, GeometryCollection::EEditUpdate InEditUpdate, bool bShapeIsUnchanged)
 	: Component(InComponent)
 	, EditUpdate(InEditUpdate)
+	, bShapeIsUnchanged(bShapeIsUnchanged)
 {
 	bHadPhysicsState = Component->HasValidPhysicsState();
 	if (EnumHasAnyFlags(EditUpdate, GeometryCollection::EEditUpdate::Physics) && bHadPhysicsState)
@@ -2258,7 +2259,10 @@ FGeometryCollectionEdit::~FGeometryCollectionEdit()
 
 		if (EnumHasAnyFlags(EditUpdate, GeometryCollection::EEditUpdate::Rest) && GetRestCollection())
 		{
-			GetRestCollection()->UpdateConvexGeometry();
+			if (!bShapeIsUnchanged)
+			{
+				GetRestCollection()->UpdateConvexGeometry();
+			}
 			GetRestCollection()->InvalidateCollection();
 		}
 
diff --git a/Engine/Source/Runtime/Experimental/GeometryCollectionEngine/Public/GeometryCollection/GeometryCollectionComponent.h b/Engine/Source/Runtime/Experimental/GeometryCollectionEngine/Public/GeometryCollection/GeometryCollectionComponent.h
--- a/Engine/Source/Runtime/Experimental/GeometryCollectionEngine/Public/GeometryCollection/GeometryCollectionComponent.h
+++ b/Engine/Source/Runtime/Experimental/GeometryCollectionEngine/Public/GeometryCollection/GeometryCollectionComponent.h
@@ -177,7 +177,12 @@ namespace GeometryCollection
 class GEOMETRYCOLLECTIONENGINE_API FGeometryCollectionEdit
 {
 public:
-	FGeometryCollectionEdit(UGeometryCollectionComponent* InComponent, GeometryCollection::EEditUpdate EditUpdate = GeometryCollection::EEditUpdate::RestPhysics);
+	/**
+	 * @param InComponent				The component to edit
+	 * @param EditUpdate				What parts of the geometry collection to update
+	 * @param bShapeIsUnchanged			Override indicating the overall shape of the geometry and clusters is unchanged, even if the rest collection changed.  Useful to e.g., not re-compute convex hulls when we don't need to.
+	 */
+	FGeometryCollectionEdit(UGeometryCollectionComponent* InComponent, GeometryCollection::EEditUpdate EditUpdate = GeometryCollection::EEditUpdate::RestPhysics, bool bShapeIsUnchanged = false);
 	~FGeometryCollectionEdit();
 
 	UGeometryCollection* GetRestCollection();
@@ -186,6 +191,7 @@ private:
 	UGeometryCollectionComponent* Component;
 	const GeometryCollection::EEditUpdate EditUpdate;
 	bool bHadPhysicsState;
+	bool bShapeIsUnchanged;
 };
 
 #if WITH_EDITOR
@@ -384,7 +390,7 @@ public:
 	/** RestCollection */
 	void SetRestCollection(const UGeometryCollection * RestCollectionIn);
 	FORCEINLINE const UGeometryCollection* GetRestCollection() const { return RestCollection; }
-	FORCEINLINE FGeometryCollectionEdit EditRestCollection(GeometryCollection::EEditUpdate EditUpdate = GeometryCollection::EEditUpdate::RestPhysics) { return FGeometryCollectionEdit(this, EditUpdate); }
+	FORCEINLINE FGeometryCollectionEdit EditRestCollection(GeometryCollection::EEditUpdate EditUpdate = GeometryCollection::EEditUpdate::RestPhysics, bool bShapeIsUnchanged = false) { return FGeometryCollectionEdit(this, EditUpdate, bShapeIsUnchanged); }
 #if WITH_EDITOR
 	FORCEINLINE FScopedColorEdit EditBoneSelection() { return FScopedColorEdit(this); }
 
