make interactive selection tool's volume-based selection consider all bones, including clusters -- not just those with geometry

#rnx

diff --git a/Engine/Plugins/Experimental/ChaosEditor/Source/FractureEditor/Private/FractureToolSelection.cpp b/Engine/Plugins/Experimental/ChaosEditor/Source/FractureEditor/Private/FractureToolSelection.cpp
--- a/Engine/Plugins/Experimental/ChaosEditor/Source/FractureEditor/Private/FractureToolSelection.cpp
+++ b/Engine/Plugins/Experimental/ChaosEditor/Source/FractureEditor/Private/FractureToolSelection.cpp
@@ -342,59 +342,55 @@ void UFractureToolSelection::FractureContextChanged()
 		FGeometryCollection& Collection = *FractureContext.GetGeometryCollection();
 		int CollectionIdx = VisualizedCollections.Emplace(FractureContext.GetGeometryCollectionComponent());
 
-		TArray<double> Volumes;
-		FindBoneVolumes(
-			Collection,
-			TArrayView<int32>(), /*Empty array => use all transforms*/
-			Volumes,
-			VolDimScale
-		);
-
-		TInterval<double> VolumeRange = GetVolumeRange(Volumes);
-
-
 		TArray<int32> SelectedIndices;
-		FilterBonesByVolume(
-			Collection,
-			TArrayView<int32>(),
-			Volumes,
-			[&VolumeRange](double Volume, int32 Bone)
-			{
-				return VolumeRange.Contains(Volume);
-			},
-			SelectedIndices
-		);
-
-		FTransform OuterTransform = FractureContext.GetTransform();
-		for (int32 TransformIdx : SelectedIndices)
+		if (!GetBonesByVolume(Collection, SelectedIndices))
 		{
-			FTransform InnerTransform = GeometryCollectionAlgo::GlobalMatrix(Collection.Transform, Collection.Parent, TransformIdx);
-
+			continue;
+		}
 
-			FTransform CombinedTransform = InnerTransform * OuterTransform;
+		FTransform OuterTransform = FractureContext.GetTransform();
 
-			int32 GeometryIdx = Collection.TransformToGeometryIndex[TransformIdx];
-			if (GeometryIdx == INDEX_NONE)
+		// Build bounding boxes for every transform (including clusters) in the global space of the component
+		TArray<FTransform> AllTransforms;
+		TArray<FBox> AllBounds;
+		int32 NumBones = Collection.NumElements(FGeometryCollection::TransformGroup);
+		AllBounds.SetNum(NumBones);
+		GeometryCollectionAlgo::GlobalMatrices(Collection.Transform, Collection.Parent, AllTransforms);
+		FBox EmptyBounds(EForceInit::ForceInit);
+		AllBounds.Init(EmptyBounds, NumBones);
+		TArray<int32> BoneIndices = GeometryCollectionAlgo::ComputeRecursiveOrder(Collection);
+		for (int32 BoneIdx : BoneIndices)
+		{
+			int32 GeoIdx = Collection.TransformToGeometryIndex[BoneIdx];
+			if (GeoIdx != INDEX_NONE)
 			{
-				ensure(false); // already filtered for existing geo, shouldn't have any missing-geo transforms in the list
-				continue;
+				const FTransform& InnerTransform = AllTransforms[BoneIdx];
+				FTransform CombinedTransform = InnerTransform * OuterTransform;
+
+				AllBounds[BoneIdx] += Collection.BoundingBox[GeoIdx].TransformBy(CombinedTransform);
 			}
-			int32 VStart = Collection.VertexStart[GeometryIdx];
-			int32 VEnd = VStart + Collection.VertexCount[GeometryIdx];
-			FBox Bounds(EForceInit::ForceInit);
-			for (int32 VIdx = VStart; VIdx < VEnd; VIdx++)
+			int32 ParentIdx = Collection.Parent[BoneIdx];
+			if (ParentIdx != INDEX_NONE)
 			{
-				Bounds += CombinedTransform.TransformPosition(Collection.Vertex[VIdx]);
+				AllBounds[ParentIdx] += AllBounds[BoneIdx];
 			}
+		}
+
+		// Add bounding boxes for the selected indices to the visualization
+		for (int32 TransformIdx : SelectedIndices)
+		{
+			FBox Bounds = AllBounds[TransformIdx];
 			SelectionMappings.AddMapping(CollectionIdx, TransformIdx, SelectionBounds.Num());
 			SelectionBounds.Add(Bounds);
 		}
 	}
 }
 
-TInterval<double> UFractureToolSelection::GetVolumeRange(TArray<double>& Volumes)
+TInterval<float> UFractureToolSelection::GetVolumeRange(const TManagedArray<float>& Volumes, const TManagedArray<int32>& TransformToGeometryIndex)
 {
-	TInterval<double> Range;
+	check(Volumes.Num() == TransformToGeometryIndex.Num());
+
+	TInterval<float> Range;
 	if (SelectionSettings->VolumeSelectionMethod == EVolumeSelectionMethod::CubeRoot)
 	{
 		Range.Min = SelectionSettings->MinVolume * VolDimScale;
@@ -405,9 +401,13 @@ TInterval<double> UFractureToolSelection::GetVolumeRange(TArray<double>& Volumes
 	else if (SelectionSettings->VolumeSelectionMethod == EVolumeSelectionMethod::RelativeToLargest)
 	{
 		double LargestVolume = KINDA_SMALL_NUMBER;
-		for (double Volume : Volumes)
+		for (int32 Idx = 0; Idx < Volumes.Num(); Idx++)
 		{
-			LargestVolume = FMath::Max(LargestVolume, Volume);
+			if (TransformToGeometryIndex[Idx] >= 0)
+			{
+				float Volume = Volumes[Idx];
+				LargestVolume = FMath::Max(LargestVolume, Volume);
+			}
 		}
 		Range.Min = LargestVolume * SelectionSettings->MinVolumeFrac * SelectionSettings->MinVolumeFrac * SelectionSettings->MinVolumeFrac;
 		Range.Max = LargestVolume * SelectionSettings->MaxVolumeFrac * SelectionSettings->MaxVolumeFrac * SelectionSettings->MaxVolumeFrac;
@@ -415,9 +415,13 @@ TInterval<double> UFractureToolSelection::GetVolumeRange(TArray<double>& Volumes
 	else // EVolumeSelectionMethod::RelativeToWhole
 	{
 		double VolumeSum = 0;
-		for (double Volume : Volumes)
+		for (int32 Idx = 0; Idx < Volumes.Num(); Idx++)
 		{
-			VolumeSum += Volume;
+			if (TransformToGeometryIndex[Idx] >= 0)
+			{
+				float Volume = Volumes[Idx];
+				VolumeSum += Volume;
+			}
 		}
 		Range.Min = VolumeSum * SelectionSettings->MinVolumeFrac * SelectionSettings->MinVolumeFrac * SelectionSettings->MinVolumeFrac;
 		Range.Max = VolumeSum * SelectionSettings->MaxVolumeFrac * SelectionSettings->MaxVolumeFrac * SelectionSettings->MaxVolumeFrac;
@@ -425,6 +429,26 @@ TInterval<double> UFractureToolSelection::GetVolumeRange(TArray<double>& Volumes
 	return Range;
 }
 
+bool UFractureToolSelection::GetBonesByVolume(const FGeometryCollection& Collection, TArray<int32>& FilterIndices)
+{
+	const TManagedArray<float>* Volumes = Collection.FindAttribute<float>("Volume", FGeometryCollection::TransformGroup);
+	if (!Volumes)
+	{
+		UE_LOG(LogFractureTool, Warning, TEXT("Volumes missing from geometry collection; could not filter selection by volume."));
+		return false;
+	}
+
+	TInterval<float> VolumeRange = GetVolumeRange(*Volumes, Collection.TransformToGeometryIndex);
+
+	for (int32 BoneIdx = 0; BoneIdx < Volumes->Num(); BoneIdx++)
+	{
+		if (VolumeRange.Contains((*Volumes)[BoneIdx]))
+		{
+			FilterIndices.Add(BoneIdx);
+		}
+	}
+	return true;
+}
 
 void UFractureToolSelection::Execute(TWeakPtr<FFractureEditorModeToolkit> InToolkit)
 {
@@ -457,27 +481,11 @@ void UFractureToolSelection::UpdateSelection(FFractureToolContext& FractureConte
 
 	FGeometryCollection& Collection = *FractureContext.GetGeometryCollection();
 
-	TArray<double> Volumes;
-	FindBoneVolumes(
-		Collection,
-		TArrayView<int32>(), /*Empty array => use all transforms*/
-		Volumes,
-		VolDimScale
-	);
-
-	TInterval<double> VolumeRange = GetVolumeRange(Volumes);
-
 	TArray<int32> FilterIndices;
-	FilterBonesByVolume(
-		Collection,
-		TArrayView<int32>(),
-		Volumes,
-		[&VolumeRange](double Volume, int32 Bone)
-		{
-			return VolumeRange.Contains(Volume);
-		},
-		FilterIndices
-	);
+	if (!GetBonesByVolume(Collection, FilterIndices))
+	{
+		return;
+	}
 
 	UFractureEditorMode* FractureMode = Cast<UFractureEditorMode>(GLevelEditorModeTools().GetActiveScriptableMode(UFractureEditorMode::EM_FractureEditorModeId));
 	if (FractureMode)
diff --git a/Engine/Plugins/Experimental/ChaosEditor/Source/FractureEditor/Private/FractureToolSelection.h b/Engine/Plugins/Experimental/ChaosEditor/Source/FractureEditor/Private/FractureToolSelection.h
--- a/Engine/Plugins/Experimental/ChaosEditor/Source/FractureEditor/Private/FractureToolSelection.h
+++ b/Engine/Plugins/Experimental/ChaosEditor/Source/FractureEditor/Private/FractureToolSelection.h
@@ -267,7 +267,8 @@ private:
 	TArray<FBox> SelectionBounds; // Bounds in global space but without exploded vectors applied
 	FVisualizationMappings SelectionMappings;
 
-	TInterval<double> GetVolumeRange(TArray<double>& Volumes);
+	TInterval<float> GetVolumeRange(const TManagedArray<float>& Volumes, const TManagedArray<int32>& TransformToGeometryIndex);
+	bool GetBonesByVolume(const FGeometryCollection& Collection, TArray<int32>& FilterIndices);
 	const double VolDimScale = .01; // compute volumes in meters instead of cm, for saner units at typical scales
 
 };
