cis fixes -- add missing includes in fracture code to fix non-unity compile errors

#rnx

diff --git a/Engine/Plugins/Experimental/ChaosEditor/Source/FractureEditor/Private/FractureToolUV.cpp b/Engine/Plugins/Experimental/ChaosEditor/Source/FractureEditor/Private/FractureToolUV.cpp
--- a/Engine/Plugins/Experimental/ChaosEditor/Source/FractureEditor/Private/FractureToolUV.cpp
+++ b/Engine/Plugins/Experimental/ChaosEditor/Source/FractureEditor/Private/FractureToolUV.cpp
@@ -17,6 +17,7 @@
 #include "IContentBrowserSingleton.h"
 
 #include "Misc/ScopedSlowTask.h"
+#include "ScopedTransaction.h"
 
 #include "FractureAutoUV.h"
 
