move volume attribute calc out of UI code so we can get up to date volume data when computing convexes

#rnx

diff --git a/Engine/Plugins/Experimental/ChaosEditor/Source/FractureEditor/Private/FractureEditorModeToolkit.cpp b/Engine/Plugins/Experimental/ChaosEditor/Source/FractureEditor/Private/FractureEditorModeToolkit.cpp
--- a/Engine/Plugins/Experimental/ChaosEditor/Source/FractureEditor/Private/FractureEditorModeToolkit.cpp
+++ b/Engine/Plugins/Experimental/ChaosEditor/Source/FractureEditor/Private/FractureEditorModeToolkit.cpp
@@ -36,6 +36,7 @@
 
 #include "GeometryCollection/GeometryCollectionActor.h"
 #include "GeometryCollection/GeometryCollection.h"
+#include "GeometryCollection/GeometryCollectionConvexUtility.h"
 
 #include "Styling/AppStyle.h"
 #include "EditorStyleSet.h"
@@ -1486,128 +1487,20 @@ void FFractureEditorModeToolkit::UpdateGeometryComponentAttributes(UGeometryColl
 		if (RestCollection && IsValidChecked(RestCollection))
 		{
 			FGeometryCollectionPtr GeometryCollection = RestCollection->GetGeometryCollection();
-
 			if (!GeometryCollection->HasAttribute("Volume", FTransformCollection::TransformGroup))
 			{
-				GeometryCollection->AddAttribute<float>("Volume", FTransformCollection::TransformGroup);
+				// Note: SetVolumeAttributes (below) will add the attribute as needed
 				UE_LOG(LogFractureTool, Warning, TEXT("Added Volume attribute to GeometryCollection."));
 			}
 
-			TArray<FTransform> Transform;
-			GeometryCollectionAlgo::GlobalMatrices(GeometryCollection->Transform, GeometryCollection->Parent, Transform);
-
-			const TManagedArray<FVector3f>& Vertex = GeometryCollection->Vertex;
-			const TManagedArray<int32>& BoneMap = GeometryCollection->BoneMap;
-
-			Chaos::TParticles<Chaos::FReal, 3> MassSpaceParticles;
-			MassSpaceParticles.AddParticles(Vertex.Num());
-			for (int32 Idx = 0; Idx < Vertex.Num(); ++Idx)
-			{
-				MassSpaceParticles.X(Idx) = Transform[BoneMap[Idx]].TransformPosition(Vertex[Idx]);
-			}
-
-			// recalculate mass recursively from root nodes
-			const TManagedArray<int32>& Parent = GeometryCollection->Parent;
-			for (int32 Idx = 0; Idx < Parent.Num(); ++Idx)
-			{
-				if (Parent[Idx] == INDEX_NONE)
-				{
-					UpdateVolumes(GeometryCollection, MassSpaceParticles, Idx);
-				}
-			}
-
-			// Normalize Size
-			if (!GeometryCollection->HasAttribute("Size", FTransformCollection::TransformGroup))
-			{
-				GeometryCollection->AddAttribute<float>("Size", FTransformCollection::TransformGroup);
-				UE_LOG(LogFractureTool, Warning, TEXT("Added Size attribute to GeometryCollection."));
-			}
-			TManagedArray<float>& Size = GeometryCollection->GetAttribute<float>("Size", FTransformCollection::TransformGroup);
-			const TManagedArray<int32>& SimulationType = GeometryCollection->SimulationType;
-			float MaxSize = 0.0f;
-			for (int32 Idx = 0; Idx < Size.Num(); ++Idx)
-			{
-				// We use cube root because it's a more manageable way to look at the numbers
-				Size[Idx] = FGenericPlatformMath::Pow(static_cast<float>(Size[Idx]), (1.0f / 3.0f));
-				
-				if (SimulationType[Idx] == FGeometryCollection::ESimulationTypes::FST_Rigid)
-				{ 
-					MaxSize = FMath::Max(Size[Idx], MaxSize);
-				}
-			}
-			if (MaxSize > 0.0)
-			{ 
-				for (int32 Idx = 0; Idx < Size.Num(); ++Idx)
-				{ 
-					Size[Idx] /= MaxSize;
-				}
-			}
+			// TODO: this should instead be called systematically in FGeometryCollectionEdit or similar
+			// (it is currently also called by the convex generation, however it is relatively fast so is ok if we call it twice)
+			FGeometryCollectionConvexUtility::SetVolumeAttributes(GeometryCollection.Get());
 		}
 	}
 	
 }
 
-void FFractureEditorModeToolkit::UpdateVolumes(FGeometryCollectionPtr GeometryCollection, const Chaos::TParticles<Chaos::FReal, 3>& MassSpaceParticles, int32 TransformIndex)
-{
-	const TManagedArray<TSet<int32>>& Children = GeometryCollection->Children;
-	const TManagedArray<int32>& SimulationType = GeometryCollection->SimulationType;
-	const TManagedArray<bool>& Visible = GeometryCollection->Visible;
-	const TManagedArray<int32>& FaceCount = GeometryCollection->FaceCount;
-	const TManagedArray<int32>& FaceStart = GeometryCollection->FaceStart;
-	const TManagedArray<int32>& TransformToGeometryIndex = GeometryCollection->TransformToGeometryIndex;
-	const TManagedArray<FIntVector>& Indices = GeometryCollection->Indices;
-	TManagedArray<float>& Volumes = GeometryCollection->GetAttribute<float>("Volume", FTransformCollection::TransformGroup);
-	
-	if (!GeometryCollection->HasAttribute("Size", FTransformCollection::TransformGroup))
-	{
-		GeometryCollection->AddAttribute<float>("Size", FTransformCollection::TransformGroup);
-	}
-	TManagedArray<float>& Size = GeometryCollection->GetAttribute<float>("Size", FTransformCollection::TransformGroup);
-
-	if (SimulationType[TransformIndex] == FGeometryCollection::ESimulationTypes::FST_Rigid)
-	{
-		int32 GeometryIndex = TransformToGeometryIndex[TransformIndex];
-
-		if (ensureMsgf(GeometryIndex > INDEX_NONE, TEXT("Leaf node %d expected to map to geometry but did not."), TransformIndex))
-		{
-			TUniquePtr<Chaos::FTriangleMesh> TriMesh(
-				CreateTriangleMesh(
-					FaceStart[GeometryIndex],
-					FaceCount[GeometryIndex],
-					Visible,
-					Indices,
-					true));
-
-			Chaos::FReal Volume = 0.0;
-			Chaos::FVec3 CenterOfMass;
-			Chaos::CalculateVolumeAndCenterOfMass(MassSpaceParticles, TriMesh->GetElements(), Volume, CenterOfMass);
-
-			// Since we're only interested in relative mass, we assume density = 1.0
-			Volumes[TransformIndex] = Volume; // todo(lwc) potential conversion from double to float
-			Size[TransformIndex] = Volume;
-		}	
-	}
-	else if (SimulationType[TransformIndex] == FGeometryCollection::ESimulationTypes::FST_Clustered)
-	{
-		// Recurse to children and sum the volumes for this node
-		float LocalVolume = 0.0;
-		float LocalSize = 0.0;
-		for (int32 ChildIndex : Children[TransformIndex])
-		{
-			UpdateVolumes(GeometryCollection, MassSpaceParticles, ChildIndex);
-			LocalVolume += Volumes[ChildIndex];
-			LocalSize += Size[ChildIndex];
-		}
-
-		Volumes[TransformIndex] = LocalVolume;
-		Size[TransformIndex] = LocalSize;
-	}
-	else
-	{
-		// Node is embedded geometry, for which we do not require volume calculations.
-		Volumes[TransformIndex] = 0.0;
-	}
-}
 
 bool GetValidGeoCenter(FGeometryCollection* Collection, const TManagedArray<int32>& TransformToGeometryIndex, const TArray<FTransform>& Transforms, const TManagedArray<TSet<int32>>& Children, const TManagedArray<FBox>& BoundingBox, int32 TransformIndex, FVector& OutGeoCenter )
 {
diff --git a/Engine/Source/Runtime/Experimental/Chaos/Private/GeometryCollection/GeometryCollectionConvexUtility.cpp b/Engine/Source/Runtime/Experimental/Chaos/Private/GeometryCollection/GeometryCollectionConvexUtility.cpp
--- a/Engine/Source/Runtime/Experimental/Chaos/Private/GeometryCollection/GeometryCollectionConvexUtility.cpp
+++ b/Engine/Source/Runtime/Experimental/Chaos/Private/GeometryCollection/GeometryCollectionConvexUtility.cpp
@@ -1055,6 +1055,47 @@ bool CopyHulls(
 	return true;
 }
 
+// helper to compute the volume of an individual piece of geometry
+double ComputeGeometryVolume(
+	const FGeometryCollection* Collection,
+	int32 GeometryIdx,
+	const FTransform& GlobalTransform,
+	double ScalePerDimension
+)
+{
+	int32 VStart = Collection->VertexStart[GeometryIdx];
+	int32 VEnd = VStart + Collection->VertexCount[GeometryIdx];
+	if (VStart == VEnd)
+	{
+		return 0.0;
+	}
+	FVector3d Center = FVector::ZeroVector;
+	for (int32 VIdx = VStart; VIdx < VEnd; VIdx++)
+	{
+		FVector Pos = GlobalTransform.TransformPosition(Collection->Vertex[VIdx]);
+		Center += (FVector3d)Pos;
+	}
+	Center /= double(VEnd - VStart);
+	int32 FStart = Collection->FaceStart[GeometryIdx];
+	int32 FEnd = FStart + Collection->FaceCount[GeometryIdx];
+	double VolOut = 0;
+	for (int32 FIdx = FStart; FIdx < FEnd; FIdx++)
+	{
+		FIntVector Tri = Collection->Indices[FIdx];
+		FVector3d V0 = (FVector3d)GlobalTransform.TransformPosition(Collection->Vertex[Tri.X]);
+		FVector3d V1 = (FVector3d)GlobalTransform.TransformPosition(Collection->Vertex[Tri.Y]);
+		FVector3d V2 = (FVector3d)GlobalTransform.TransformPosition(Collection->Vertex[Tri.Z]);
+
+		// add volume of the tetrahedron formed by the triangles and the reference point
+		FVector3d V1mRef = (V1 - Center) * ScalePerDimension;
+		FVector3d V2mRef = (V2 - Center) * ScalePerDimension;
+		FVector3d N = V2mRef.Cross(V1mRef);
+
+		VolOut += ((V0 - Center) * ScalePerDimension).Dot(N) / 6.0;
+	}
+	return VolOut;
+}
+
 
 }
 
@@ -1073,6 +1114,8 @@ FGeometryCollectionConvexUtility::FGeometryCollectionConvexData FGeometryCollect
 		FGeometryCollectionProximityUtility ProximityUtility(GeometryCollection);
 		ProximityUtility.UpdateProximity();
 	}
+	SetVolumeAttributes(GeometryCollection);
+
 	const TManagedArray<TSet<int32>>* GCProximity = GeometryCollection->FindAttribute<TSet<int32>>("Proximity", FGeometryCollection::GeometryGroup);
 	const TManagedArray<float>* Volume = GeometryCollection->FindAttribute<float>("Volume", FGeometryCollection::TransformGroup);
 	TArray<TUniquePtr<Chaos::FConvex>> Convexes;
@@ -1342,3 +1385,55 @@ void FGeometryCollectionConvexUtility::CopyChildConvexes(const FGeometryCollecti
 	}
 }
 
+void FGeometryCollectionConvexUtility::SetVolumeAttributes(FGeometryCollection* Collection)
+{
+	if (!Collection->HasAttribute("Volume", FGeometryCollection::TransformGroup))
+	{
+		Collection->AddAttribute<float>("Volume", FGeometryCollection::TransformGroup);
+	}
+	if (!Collection->HasAttribute("Size", FGeometryCollection::TransformGroup))
+	{
+		Collection->AddAttribute<float>("Size", FGeometryCollection::TransformGroup);
+	}
+	TManagedArray<float>& Volumes = Collection->GetAttribute<float>("Volume", FTransformCollection::TransformGroup);
+	TManagedArray<float>& Sizes = Collection->GetAttribute<float>("Size", FTransformCollection::TransformGroup);
+
+	const TManagedArray<int32>& SimulationType = Collection->SimulationType;
+	const TManagedArray<int32>& TransformToGeometryIndex = Collection->TransformToGeometryIndex;
+	const TManagedArray<int32>& Parent = Collection->Parent;
+
+	TArray<FTransform> Transforms;
+	GeometryCollectionAlgo::GlobalMatrices(Collection->Transform, Collection->Parent, Transforms);
+	TArray<int32> RecursiveOrder = GeometryCollectionAlgo::ComputeRecursiveOrder(*Collection);
+	float MaxGeoVolume = 0.0f;
+	for (int32 Bone : RecursiveOrder)
+	{
+		if (SimulationType[Bone] == FGeometryCollection::ESimulationTypes::FST_Rigid)
+		{
+			int32 GeoIdx = TransformToGeometryIndex[Bone];
+			if (GeoIdx == INDEX_NONE)
+			{
+				Volumes[Bone] = 0.0f;
+				continue;
+			}
+			else
+			{
+				float GeoVolume = (float)ComputeGeometryVolume(Collection, GeoIdx, Transforms[Bone], 1.0);
+				MaxGeoVolume = FMath::Max(MaxGeoVolume, GeoVolume);
+				Volumes[Bone] = GeoVolume;
+			}
+		}
+		int32 ParentIdx = Parent[Bone];
+		if (ParentIdx != INDEX_NONE)
+		{
+			Volumes[ParentIdx] += Volumes[Bone];
+		}
+	}
+	float SizeScaleFactor = MaxGeoVolume > 0 ? (1.0f / FGenericPlatformMath::Pow(MaxGeoVolume, (1.0f / 3.0f))) : 1.0f;
+	for (int32 BoneIdx = 0; BoneIdx < Volumes.Num(); BoneIdx++)
+	{
+		float Rt3Volume = FGenericPlatformMath::Pow(Volumes[BoneIdx], 1.0f / 3.0f);
+		Sizes[BoneIdx] = Rt3Volume * SizeScaleFactor;
+	}
+}
+
diff --git a/Engine/Source/Runtime/Experimental/Chaos/Public/GeometryCollection/GeometryCollectionConvexUtility.h b/Engine/Source/Runtime/Experimental/Chaos/Public/GeometryCollection/GeometryCollectionConvexUtility.h
--- a/Engine/Source/Runtime/Experimental/Chaos/Public/GeometryCollection/GeometryCollectionConvexUtility.h
+++ b/Engine/Source/Runtime/Experimental/Chaos/Public/GeometryCollection/GeometryCollectionConvexUtility.h
@@ -56,6 +56,9 @@ public:
 	/** @return true if the GeometryCollection has convex data with no null pointers and no invalid indices */
 	static bool ValidateConvexData(const FGeometryCollection* GeometryCollection);
 
+	/** Set Volume and Size attributes on the Collection (will be called by CreateNonOverlappingConvexHullData -- Volumes must be up to date for convex calc) */
+	static void SetVolumeAttributes(FGeometryCollection* Collection);
+
 	/**
 	 * Copy convex hulls from *below* FromTransformIdx over to all live at ToTransformIdx.
 	 * The two geometry collections can be the same but do not need to be.
