also remove fractureeditormodetoolkit deleted UpdateVolumes function from header

#rnx

diff --git a/Engine/Plugins/Experimental/ChaosEditor/Source/FractureEditor/Public/FractureEditorModeToolkit.h b/Engine/Plugins/Experimental/ChaosEditor/Source/FractureEditor/Public/FractureEditorModeToolkit.h
--- a/Engine/Plugins/Experimental/ChaosEditor/Source/FractureEditor/Public/FractureEditorModeToolkit.h
+++ b/Engine/Plugins/Experimental/ChaosEditor/Source/FractureEditor/Public/FractureEditorModeToolkit.h
@@ -186,7 +186,6 @@ protected:
 
 private:
 	static void UpdateGeometryComponentAttributes(UGeometryCollectionComponent* Component);
-	static void UpdateVolumes(FGeometryCollectionPtr GeometryCollection, const Chaos::TParticles<Chaos::FReal, 3>& MassSpaceParticles, int32 TransformIndex);
 
 	void OnOutlinerBoneSelectionChanged(UGeometryCollectionComponent* RootComponent, TArray<int32>& SelectedBones);
 	void OnHistogramBoneSelectionChanged(UGeometryCollectionComponent* RootComponent, TArray<int32>& SelectedBones);
