fix fracture mode outliner showing duplicate entries for non-group fracture

#rnx

diff --git a/Engine/Plugins/Experimental/ChaosEditor/Source/FractureEditor/Public/FractureTool.cpp b/Engine/Plugins/Experimental/ChaosEditor/Source/FractureEditor/Public/FractureTool.cpp
--- a/Engine/Plugins/Experimental/ChaosEditor/Source/FractureEditor/Public/FractureTool.cpp
+++ b/Engine/Plugins/Experimental/ChaosEditor/Source/FractureEditor/Public/FractureTool.cpp
@@ -206,7 +206,7 @@ void UFractureActionTool::SetOutlinerComponents(TArray<FFractureToolContext>& In
 	TArray<UGeometryCollectionComponent*> Components;
 	for (FFractureToolContext& Context : InContexts)
 	{
-		Components.Add(Context.GetGeometryCollectionComponent());
+		Components.AddUnique(Context.GetGeometryCollectionComponent());
 	}
 	Toolkit->SetOutlinerComponents(Components);
 }
