fix crash in fracture mode cluster magnet tool

diff --git a/Engine/Plugins/Experimental/ChaosEditor/Source/FractureEditor/Private/FractureToolClusterMagnet.cpp b/Engine/Plugins/Experimental/ChaosEditor/Source/FractureEditor/Private/FractureToolClusterMagnet.cpp
--- a/Engine/Plugins/Experimental/ChaosEditor/Source/FractureEditor/Private/FractureToolClusterMagnet.cpp
+++ b/Engine/Plugins/Experimental/ChaosEditor/Source/FractureEditor/Private/FractureToolClusterMagnet.cpp
@@ -189,7 +189,8 @@ TMap<int32, TSet<int32>> UFractureToolClusterMagnet::InitializeConnectivity(cons
 void UFractureToolClusterMagnet::CollectTopNodeConnections(FGeometryCollectionPtr GeometryCollection, int32 Index, int32 OperatingLevel, TSet<int32>& OutConnections) const
 {
 	const TManagedArray<int32>& TransformToGeometryIndex = GeometryCollection->TransformToGeometryIndex;
-	if (TransformToGeometryIndex[Index] == INDEX_NONE) // leaf node
+	if (GeometryCollection->SimulationType[Index] == FGeometryCollection::ESimulationTypes::FST_Rigid
+		&& TransformToGeometryIndex[Index] != INDEX_NONE) // rigid node with geometry, leaf of the simulated part
 	{
 		const TManagedArray<TSet<int32>>& Proximity = GeometryCollection->GetAttribute<TSet<int32>>("Proximity", FGeometryCollection::GeometryGroup);
 		const TManagedArray<int32>& GeometryToTransformIndex = GeometryCollection->TransformIndex;
@@ -208,7 +209,7 @@ void UFractureToolClusterMagnet::CollectTopNodeConnections(FGeometryCollectionPt
 		{
 			CollectTopNodeConnections(GeometryCollection, ChildIndex, OperatingLevel, OutConnections);
 		}
-	}	
+	}
 }
 
 void UFractureToolClusterMagnet::SeparateClusterMagnets(
