fix fracture mode selecting the wrong bones after tool actions

#rnx

diff --git a/Engine/Plugins/Experimental/ChaosEditor/Source/FractureEditor/Private/FractureEditorModeToolkit.cpp b/Engine/Plugins/Experimental/ChaosEditor/Source/FractureEditor/Private/FractureEditorModeToolkit.cpp
--- a/Engine/Plugins/Experimental/ChaosEditor/Source/FractureEditor/Private/FractureEditorModeToolkit.cpp
+++ b/Engine/Plugins/Experimental/ChaosEditor/Source/FractureEditor/Private/FractureEditorModeToolkit.cpp
@@ -1233,6 +1233,19 @@ void FFractureEditorModeToolkit::ExecuteAction(UFractureActionTool* InActionTool
 	if (InActionTool)
 	{
 		InActionTool->Execute(StaticCastSharedRef<FFractureEditorModeToolkit>(AsShared()));
+
+		InvalidateHitProxies();
+	}
+}
+
+void FFractureEditorModeToolkit::InvalidateHitProxies()
+{
+	if (GIsEditor)
+	{
+		for (FEditorViewportClient* Viewport : GEditor->GetLevelViewportClients())
+		{
+			Viewport->Invalidate();
+		}
 	}
 }
 
diff --git a/Engine/Plugins/Experimental/ChaosEditor/Source/FractureEditor/Public/FractureEditorModeToolkit.h b/Engine/Plugins/Experimental/ChaosEditor/Source/FractureEditor/Public/FractureEditorModeToolkit.h
--- a/Engine/Plugins/Experimental/ChaosEditor/Source/FractureEditor/Public/FractureEditorModeToolkit.h
+++ b/Engine/Plugins/Experimental/ChaosEditor/Source/FractureEditor/Public/FractureEditorModeToolkit.h
@@ -184,6 +184,9 @@ protected:
 
 	void InvalidateCachedDetailPanelState(UObject* ChangedObject);
 
+	// Invalidate the hit proxies for all level viewports; we need to do this after updating geometry collection(s)
+	void InvalidateHitProxies();
+
 private:
 	static void UpdateGeometryComponentAttributes(UGeometryCollectionComponent* Component);
 
