fix crash on fracture mode cluster merge when selection is empty + make fracture contexts always remove invalid bones from new selections

#rnx

diff --git a/Engine/Plugins/Experimental/ChaosEditor/Source/FractureEditor/Private/FractureToolClustering.cpp b/Engine/Plugins/Experimental/ChaosEditor/Source/FractureEditor/Private/FractureToolClustering.cpp
--- a/Engine/Plugins/Experimental/ChaosEditor/Source/FractureEditor/Private/FractureToolClustering.cpp
+++ b/Engine/Plugins/Experimental/ChaosEditor/Source/FractureEditor/Private/FractureToolClustering.cpp
@@ -292,11 +292,14 @@ void UFractureToolClusterMerge::Execute(TWeakPtr<FFractureEditorModeToolkit> InT
 			}
 			
 			int32 MergeNode = FGeometryCollectionClusteringUtility::PickBestNodeToMergeTo(Context.GetGeometryCollection().Get(), Context.GetSelection());
-			FGeometryCollectionClusteringUtility::ClusterBonesUnderExistingNode(Context.GetGeometryCollection().Get(), MergeNode, ChildBones);
-			FGeometryCollectionClusteringUtility::RemoveDanglingClusters(Context.GetGeometryCollection().Get());
-			
-			Context.SetSelection({MergeNode});
-			Refresh(Context, Toolkit);
+			if (MergeNode >= 0)
+			{
+				FGeometryCollectionClusteringUtility::ClusterBonesUnderExistingNode(Context.GetGeometryCollection().Get(), MergeNode, ChildBones);
+				FGeometryCollectionClusteringUtility::RemoveDanglingClusters(Context.GetGeometryCollection().Get());
+
+				Context.SetSelection({ MergeNode });
+				Refresh(Context, Toolkit);
+			}
 		}
 
 		SetOutlinerComponents(Contexts, Toolkit);
diff --git a/Engine/Plugins/Experimental/ChaosEditor/Source/FractureEditor/Private/FractureToolContext.cpp b/Engine/Plugins/Experimental/ChaosEditor/Source/FractureEditor/Private/FractureToolContext.cpp
--- a/Engine/Plugins/Experimental/ChaosEditor/Source/FractureEditor/Private/FractureToolContext.cpp
+++ b/Engine/Plugins/Experimental/ChaosEditor/Source/FractureEditor/Private/FractureToolContext.cpp
@@ -22,7 +22,7 @@ FFractureToolContext::FFractureToolContext(UGeometryCollectionComponent* InGeome
 }
 
 
-void FFractureToolContext::Sanitize()
+void FFractureToolContext::Sanitize(bool bFavorParents)
 {
 	// Ensure that selected indices are valid
 	int NumTransforms = GeometryCollection->NumElements(FGeometryCollection::TransformGroup);
@@ -31,9 +31,12 @@ void FFractureToolContext::Sanitize()
 		});
 	
 	// Ensure that children of a selected node are not also selected.
-	SelectedBones.RemoveAll([this](int32 Index) {
-		return !IsValidBone(Index) || HasSelectedAncestor(Index);
-		});
+	if (bFavorParents)
+	{
+		SelectedBones.RemoveAll([this](int32 Index) {
+			return !IsValidBone(Index) || HasSelectedAncestor(Index);
+			});
+	}
 
 	SelectedBones.Sort();
 }
diff --git a/Engine/Plugins/Experimental/ChaosEditor/Source/FractureEditor/Private/FractureToolContext.h b/Engine/Plugins/Experimental/ChaosEditor/Source/FractureEditor/Private/FractureToolContext.h
--- a/Engine/Plugins/Experimental/ChaosEditor/Source/FractureEditor/Private/FractureToolContext.h
+++ b/Engine/Plugins/Experimental/ChaosEditor/Source/FractureEditor/Private/FractureToolContext.h
@@ -18,7 +18,8 @@ public:
 	/** Generate a tool context based on the bone selection of the specified GeometryCollectionComponent. */
 	FFractureToolContext(UGeometryCollectionComponent* InGeometryCollectionComponent);
 
-	void Sanitize();
+	/** @param bFavorParents	If true, when a child and its (grand)parent are both selected, only include the top-level bone in the sanitized selection */
+	void Sanitize(bool bFavorParents = true);
 	void ConvertSelectionToLeafNodes();
 	void ConvertSelectionToRigidNodes();
 	void ConvertSelectionToEmbeddedGeometryNodes();
@@ -29,7 +30,13 @@ public:
 	
 	const TArray<int32>& GetSelection() const { return SelectedBones; }
 	TArray<int32>& GetSelection() { return SelectedBones; }
-	void SetSelection(const TArray<int32>& NewSelection) { SelectedBones = NewSelection; }
+	void SetSelection(const TArray<int32>& NewSelection)
+	{
+		SelectedBones = NewSelection;
+		// Note: currently Sanitize is called without favoring parents, so it removes invalid bones and sorts the selection
+		// but does not prevent selection of children + parents together, in case such a selection is useful
+		Sanitize(/*bFavorParents*/false);
+	}
 	UGeometryCollectionComponent* GetGeometryCollectionComponent() const { return GeometryCollectionComponent; }
 	FGeometryCollectionPtr GetGeometryCollection() const { return GeometryCollection; }
 	UGeometryCollection* GetFracturedGeometryCollection() const { return FracturedGeometryCollection; }
