Fixed duplicate key for localization in the fracture tool

diff --git a/Engine/Plugins/Experimental/ChaosEditor/Source/FractureEditor/Private/SGeometryCollectionOutliner.cpp b/Engine/Plugins/Experimental/ChaosEditor/Source/FractureEditor/Private/SGeometryCollectionOutliner.cpp
--- a/Engine/Plugins/Experimental/ChaosEditor/Source/FractureEditor/Private/SGeometryCollectionOutliner.cpp
+++ b/Engine/Plugins/Experimental/ChaosEditor/Source/FractureEditor/Private/SGeometryCollectionOutliner.cpp
@@ -70,7 +70,7 @@ void FGeometryCollectionTreeItem::GenerateContextMenu(UToolMenu* Menu, SGeometry
 	};
 
 	FToolMenuSection& Section = Menu->AddSection("State");
-	Section.AddSubMenu("SetInitialDynamicState", NSLOCTEXT("Fracture", "FractureToolSetInitialDynamicState", "Initial Dynamic State"), FText(),FNewToolMenuDelegate::CreateLambda(MakeDynamicStateMenu));
+	Section.AddSubMenu("FractureToolSetInitialDynamicStateMenu", NSLOCTEXT("Fracture", "FractureToolSetInitialDynamicStateMenu", "Initial Dynamic State"), FText(),FNewToolMenuDelegate::CreateLambda(MakeDynamicStateMenu));
 }
 
 void FGeometryCollectionTreeItemBone::OnDragEnter(FDragDropEvent const& InDragDropEvent)
