fix crash on map change when fracture mode has custom gizmo visible

diff --git a/Engine/Plugins/Experimental/ChaosEditor/Source/FractureEditor/Private/FractureToolCutter.cpp b/Engine/Plugins/Experimental/ChaosEditor/Source/FractureEditor/Private/FractureToolCutter.cpp
--- a/Engine/Plugins/Experimental/ChaosEditor/Source/FractureEditor/Private/FractureToolCutter.cpp
+++ b/Engine/Plugins/Experimental/ChaosEditor/Source/FractureEditor/Private/FractureToolCutter.cpp
@@ -26,8 +26,9 @@ UFractureTransformGizmoSettings::UFractureTransformGizmoSettings(const FObjectIn
 
 void UFractureTransformGizmoSettings::ResetGizmo(bool bResetRotation)
 {
-	if (!TransformGizmo || !TransformProxy)
+	if (!TransformGizmo || !TransformProxy || !TransformGizmo->GetGizmoActor())
 	{
+		// doesn't have a valid gizmo to reset
 		return;
 	}
 	if (!bUseGizmo || !AttachedCutter)
@@ -106,6 +107,8 @@ void UFractureTransformGizmoSettings::Shutdown()
 	{
 		Context->GizmoManager->DestroyAllGizmosByOwner(this);
 	}
+	TransformGizmo = nullptr;
+	TransformProxy = nullptr;
 }
 
 void UFractureCutterSettings::TransferNoiseSettings(FNoiseSettings& NoiseSettingsOut)
diff --git a/Engine/Plugins/Experimental/ChaosEditor/Source/FractureEditor/Private/FractureToolCutter.h b/Engine/Plugins/Experimental/ChaosEditor/Source/FractureEditor/Private/FractureToolCutter.h
--- a/Engine/Plugins/Experimental/ChaosEditor/Source/FractureEditor/Private/FractureToolCutter.h
+++ b/Engine/Plugins/Experimental/ChaosEditor/Source/FractureEditor/Private/FractureToolCutter.h
@@ -256,7 +256,7 @@ public:
 
 	bool IsGizmoEnabled()
 	{
-		return bUseGizmo && TransformGizmo->IsVisible();
+		return bUseGizmo && TransformGizmo && TransformGizmo->IsVisible();
 	}
 
 	FTransform GetTransform()
