fix incorrect transform on fracture diagram w/ custom voronoi fracture

diff --git a/Engine/Plugins/Experimental/ChaosEditor/Source/FractureEditor/Private/FractureToolCustomVoronoi.cpp b/Engine/Plugins/Experimental/ChaosEditor/Source/FractureEditor/Private/FractureToolCustomVoronoi.cpp
--- a/Engine/Plugins/Experimental/ChaosEditor/Source/FractureEditor/Private/FractureToolCustomVoronoi.cpp
+++ b/Engine/Plugins/Experimental/ChaosEditor/Source/FractureEditor/Private/FractureToolCustomVoronoi.cpp
@@ -10,6 +10,8 @@
 #include "MeshDescription.h"
 #include "Spatial/PriorityOrderPoints.h"
 
+#include "GeometryCollection/GeometryCollectionAlgo.h"
+
 #define LOCTEXT_NAMESPACE "FractureCustomVoronoi"
 
 
@@ -242,6 +244,7 @@ void UFractureToolCustomVoronoi::GenerateLivePattern(const TArray<FFractureToolC
 			FGeometryCollection& Collection = *Context.GetGeometryCollection();
 			for (int32 Bone : Context.GetSelection())
 			{
+				FTransform BoneTransform = GeometryCollectionAlgo::GlobalMatrix(Collection.Transform, Collection.Parent, Bone);
 				int32 GeometryIdx = Collection.TransformToGeometryIndex[Bone];
 				if (GeometryIdx == INDEX_NONE)
 				{
@@ -288,7 +291,7 @@ void UFractureToolCustomVoronoi::GenerateLivePattern(const TArray<FFractureToolC
 					for (TPair<FVector, int32>& Point : VertexHash)
 					{
 						const FVector OffsetPosition = (FVector)Point.Key + (FVector)Normals[Point.Value - 1] * NormalOffset;
-						const FVector LocalPos = Collection.Transform[Bone].TransformPosition(OffsetPosition);
+						const FVector LocalPos = BoneTransform.TransformPosition(OffsetPosition);
 						Sites.Add(Context.GetTransform().TransformPosition(LocalPos));
 					}
 					if (bUseOrderedSkip)
@@ -316,7 +319,7 @@ void UFractureToolCustomVoronoi::GenerateLivePattern(const TArray<FFractureToolC
 					for (int32 VIdx = Start; VIdx < Start + Count; VIdx++)
 					{
 						const FVector& Position = (FVector)Collection.Vertex[VIdx];
-						const FVector LocalPos = Collection.Transform[Bone].TransformPosition(Position);
+						const FVector LocalPos = BoneTransform.TransformPosition(Position);
 						Sites.Add(Context.GetTransform().TransformPosition(LocalPos));
 					}
 				}
