fix fracture interactive selection tool taking over selection / trapping user when there is no geometry collection selected:
- only take over the input when a geometry collection is selected + the selection mode is 'rect select'
- make 'esc' key exit the tool, so users can get out of this tool

diff --git a/Engine/Plugins/Experimental/ChaosEditor/Source/FractureEditor/Private/FractureEditorCommands.cpp b/Engine/Plugins/Experimental/ChaosEditor/Source/FractureEditor/Private/FractureEditorCommands.cpp
--- a/Engine/Plugins/Experimental/ChaosEditor/Source/FractureEditor/Private/FractureEditorCommands.cpp
+++ b/Engine/Plugins/Experimental/ChaosEditor/Source/FractureEditor/Private/FractureEditorCommands.cpp
@@ -47,6 +47,8 @@ void FFractureEditorCommands::RegisterCommands()
 	RegisterCommandInfo(ExplodeMore);
 	UI_COMMAND(ExplodeLess, "ExplodeLess", "Explode 10% Less", EUserInterfaceActionType::Button, FInputChord(EModifierKey::Shift, EKeys::Q) );
 	RegisterCommandInfo(ExplodeLess);
+	UI_COMMAND(CancelTool, "CancelTool", "Exit Current Tool Mode", EUserInterfaceActionType::Button, FInputChord(EKeys::Escape));
+	RegisterCommandInfo(CancelTool);
 
 	// Fracture Tools
 	for (TObjectIterator<UClass> ClassIterator; ClassIterator; ++ClassIterator)
diff --git a/Engine/Plugins/Experimental/ChaosEditor/Source/FractureEditor/Private/FractureEditorCommands.h b/Engine/Plugins/Experimental/ChaosEditor/Source/FractureEditor/Private/FractureEditorCommands.h
--- a/Engine/Plugins/Experimental/ChaosEditor/Source/FractureEditor/Private/FractureEditorCommands.h
+++ b/Engine/Plugins/Experimental/ChaosEditor/Source/FractureEditor/Private/FractureEditorCommands.h
@@ -42,6 +42,9 @@ class FFractureEditorCommands : public TCommands<FFractureEditorCommands>
 		TSharedPtr< FUICommandInfo > ExplodeMore;
 		TSharedPtr< FUICommandInfo > ExplodeLess;
 
+		// Tool exit
+		TSharedPtr< FUICommandInfo > CancelTool;
+
 		// Cluster Commands
 		TSharedPtr< FUICommandInfo > AutoCluster;
 		TSharedPtr< FUICommandInfo > ClusterMagnet;
diff --git a/Engine/Plugins/Experimental/ChaosEditor/Source/FractureEditor/Private/FractureEditorModeToolkit.cpp b/Engine/Plugins/Experimental/ChaosEditor/Source/FractureEditor/Private/FractureEditorModeToolkit.cpp
--- a/Engine/Plugins/Experimental/ChaosEditor/Source/FractureEditor/Private/FractureEditorModeToolkit.cpp
+++ b/Engine/Plugins/Experimental/ChaosEditor/Source/FractureEditor/Private/FractureEditorModeToolkit.cpp
@@ -850,6 +850,11 @@ void FFractureEditorModeToolkit::BindCommands()
 		EUIActionRepeatMode::RepeatEnabled
 	);
 
+	ToolkitCommands->MapAction(
+		Commands.CancelTool,
+		FExecuteAction::CreateLambda([=]() { this->SetActiveTool(nullptr); })
+	);
+
 	// Map actions of all the Fracture Tools
 	TArray<UClass*> SourceClasses = FindFractureToolClasses();
 	for (UClass* Class : SourceClasses)
diff --git a/Engine/Plugins/Experimental/ChaosEditor/Source/FractureEditor/Private/FractureToolSelection.cpp b/Engine/Plugins/Experimental/ChaosEditor/Source/FractureEditor/Private/FractureToolSelection.cpp
--- a/Engine/Plugins/Experimental/ChaosEditor/Source/FractureEditor/Private/FractureToolSelection.cpp
+++ b/Engine/Plugins/Experimental/ChaosEditor/Source/FractureEditor/Private/FractureToolSelection.cpp
@@ -21,6 +21,28 @@
 #define LOCTEXT_NAMESPACE "FractureSelection"
 
 
+namespace UE { namespace FractureToolSelectionInternal {
+
+bool IsGeometryCollectionSelected()
+{
+	if (USelection* SelectedActors = GEditor->GetSelectedActors())
+	{
+		for (FSelectionIterator Iter(*SelectedActors); Iter; ++Iter)
+		{
+			AActor* Actor = Cast<AActor>(*Iter);
+			TArray<UGeometryCollectionComponent*, TInlineAllocator<1>> GeometryComponents;
+			Actor->GetComponents<UGeometryCollectionComponent>(GeometryComponents);
+			if (GeometryComponents.Num() > 0)
+			{
+				return true;
+			}
+		}
+	}
+	return false;
+}
+
+}} // namespace UE::FractureToolSelectionInternal
+
 UFractureToolSelection::UFractureToolSelection(const FObjectInitializer& ObjInit) 
 	: Super(ObjInit) 
 {
@@ -49,7 +71,7 @@ void UFractureToolSelection::RegisterUICommand( FFractureEditorCommands* Binding
 	BindingContext->SelectCustom = UICommandInfo;
 }
 
-void UFractureToolSelection::Setup()
+void UFractureToolSelection::CreateRectangleSelectionBehavior()
 {
 	// Set things up for being able to add behaviors to live preview.
 	SelectionBehaviorSet = NewObject<UInputBehaviorSet>();
@@ -76,6 +98,19 @@ void UFractureToolSelection::Setup()
 	Context->InputRouter->RegisterSource(SelectionBehaviorSource);
 }
 
+void UFractureToolSelection::DestroyRectangleSelectionBehavior()
+{
+	if (SelectionBehaviorSource)
+	{
+		UInteractiveToolsContext* Context = GLevelEditorModeTools().GetInteractiveToolsContext();
+		Context->InputRouter->DeregisterSource(SelectionBehaviorSource);
+	}
+
+	SelectionBehaviorSet = nullptr;
+	SelectionBehaviorSource = nullptr;
+	RectangleMarqueeManager = nullptr;
+}
+
 void UFractureToolSelection::OnUpdateModifierState(int ModifierID, bool bIsOn)
 {
 	switch (ModifierID)
@@ -277,27 +312,28 @@ void URectangleMarqueeManager::DrawHUD(FCanvas* Canvas, bool bThisViewHasFocus)
 
 void UFractureToolSelection::Shutdown()
 {
-	UInteractiveToolsContext* Context = GLevelEditorModeTools().GetInteractiveToolsContext();
-	Context->InputRouter->DeregisterSource(SelectionBehaviorSource);
-
-	SelectionBehaviorSet = nullptr;
-	SelectionBehaviorSource = nullptr;
-	RectangleMarqueeManager = nullptr;
+	DestroyRectangleSelectionBehavior();
 }
 
 void UFractureToolSelection::DrawHUD(FEditorViewportClient* ViewportClient, FViewport* Viewport, const FSceneView* View, FCanvas* Canvas)
 {
 	Super::DrawHUD(ViewportClient, Viewport, View, Canvas);
-
-	const FEditorViewportClient* Focused = GLevelEditorModeTools().GetFocusedViewportClient();
-	const bool bThisViewHasFocus = (ViewportClient == Focused);
 	
-	RectangleMarqueeManager->DrawHUD(Canvas, bThisViewHasFocus);
+	if (RectangleMarqueeManager)
+	{
+		const FEditorViewportClient* Focused = GLevelEditorModeTools().GetFocusedViewportClient();
+		const bool bThisViewHasFocus = (ViewportClient == Focused);
+
+		RectangleMarqueeManager->DrawHUD(Canvas, bThisViewHasFocus);
+	}
 }
 
 void UFractureToolSelection::Render(const FSceneView* View, FViewport* Viewport, FPrimitiveDrawInterface* PDI)
 {
-	RectangleMarqueeManager->Render(View, Viewport, PDI);
+	if (RectangleMarqueeManager)
+	{
+		RectangleMarqueeManager->Render(View, Viewport, PDI);
+	}
 
 	EnumerateVisualizationMapping(SelectionMappings, SelectionBounds.Num(), [&](int32 Idx, FVector ExplodedVector)
 	{
@@ -335,7 +371,22 @@ TArray<UObject*> UFractureToolSelection::GetSettingsObjects() const
 
 void UFractureToolSelection::FractureContextChanged()
 {
-	RectangleMarqueeManager->SetIsEnabled(SelectionSettings->MouseSelectionMethod == EMouseSelectionMethod::RectSelect);
+	bool bWantsRectangleSelection =
+		SelectionSettings->MouseSelectionMethod == EMouseSelectionMethod::RectSelect &&
+		UE::FractureToolSelectionInternal::IsGeometryCollectionSelected();
+	if (!RectangleMarqueeManager && bWantsRectangleSelection)
+	{
+		CreateRectangleSelectionBehavior();
+	}
+	else if (RectangleMarqueeManager && !bWantsRectangleSelection)
+	{
+		DestroyRectangleSelectionBehavior();
+	}
+
+	if (RectangleMarqueeManager)
+	{
+		RectangleMarqueeManager->SetIsEnabled(bWantsRectangleSelection);
+	}
 
 	UpdateDefaultRandomSeed();
 	TArray<FFractureToolContext> FractureContexts = GetFractureToolContexts();
diff --git a/Engine/Plugins/Experimental/ChaosEditor/Source/FractureEditor/Private/FractureToolSelection.h b/Engine/Plugins/Experimental/ChaosEditor/Source/FractureEditor/Private/FractureToolSelection.h
--- a/Engine/Plugins/Experimental/ChaosEditor/Source/FractureEditor/Private/FractureToolSelection.h
+++ b/Engine/Plugins/Experimental/ChaosEditor/Source/FractureEditor/Private/FractureToolSelection.h
@@ -218,7 +218,6 @@ public:
 	virtual void Execute(TWeakPtr<FFractureEditorModeToolkit> InToolkit) override;
 	void UpdateSelection(FFractureToolContext& FractureContext);
 
-	virtual void Setup() override;
 	virtual void Shutdown() override;
 
 	// IClickBehaviorTarget implementation
@@ -229,6 +228,9 @@ public:
 	virtual void OnClicked(const FInputDeviceRay& ClickPos) override;
 	virtual void OnUpdateModifierState(int ModifierID, bool bIsOn) override;
 
+	void CreateRectangleSelectionBehavior();
+	void DestroyRectangleSelectionBehavior();
+
 protected:
 	virtual void ClearVisualizations() override
 	{
