fix selection issues w/ fracture ToMesh tool
 - don't auto-select the new meshes (unless user chose to)
 - don't auto-place new meshes in world (unless user chose to)
 - if user is changing bone selection, make sure component is also selected

diff --git a/Engine/Plugins/Experimental/ChaosEditor/Source/FractureEditor/Private/FractureEditorModeToolkit.cpp b/Engine/Plugins/Experimental/ChaosEditor/Source/FractureEditor/Private/FractureEditorModeToolkit.cpp
--- a/Engine/Plugins/Experimental/ChaosEditor/Source/FractureEditor/Private/FractureEditorModeToolkit.cpp
+++ b/Engine/Plugins/Experimental/ChaosEditor/Source/FractureEditor/Private/FractureEditorModeToolkit.cpp
@@ -1385,6 +1385,11 @@ void FFractureEditorModeToolkit::SetOutlinerComponents(const TArray<UGeometryCol
 
 void FFractureEditorModeToolkit::SetBoneSelection(UGeometryCollectionComponent* InRootComponent, const TArray<int32>& InSelectedBones, bool bClearCurrentSelection)
 {
+	if (InSelectedBones.Num() > 0 && !InRootComponent->IsSelected())
+	{
+		GEditor->SelectComponent(InRootComponent, true, true);
+	}
+
 	OutlinerView->SetBoneSelection(InRootComponent, InSelectedBones, bClearCurrentSelection);
 	HistogramView->SetBoneSelection(InRootComponent, InSelectedBones, bClearCurrentSelection);
 
diff --git a/Engine/Plugins/Experimental/ChaosEditor/Source/FractureEditor/Private/FractureToolConvert.cpp b/Engine/Plugins/Experimental/ChaosEditor/Source/FractureEditor/Private/FractureToolConvert.cpp
--- a/Engine/Plugins/Experimental/ChaosEditor/Source/FractureEditor/Private/FractureToolConvert.cpp
+++ b/Engine/Plugins/Experimental/ChaosEditor/Source/FractureEditor/Private/FractureToolConvert.cpp
@@ -135,11 +135,10 @@ namespace UE
 namespace FractureToolConvertInternal
 {
 // Create + position a single static mesh asset in the provided asset package
-AActor* CreateMeshAsset(FGeometryCollection& Collection, const TManagedArray<FTransform>& BoneTransforms, const FTransform& CollectionToWorld, const TConstArrayView<UMaterialInterface*>& Materials, TConstArrayView<int32> Bones, UPackage* AssetPackage, FString UniqueAssetName, bool bCenterPivot)
+void CreateMeshAsset(
+	FGeometryCollection& Collection, const TManagedArray<FTransform>& BoneTransforms, const FTransform& CollectionToWorld, const TConstArrayView<UMaterialInterface*>& Materials,
+	TConstArrayView<int32> Bones, UPackage* AssetPackage, FString UniqueAssetName, bool bCenterPivot, bool bPlaceInWorld, bool bSelectNewActors)
 {
-	UWorld* TargetWorld = GEditor->GetEditorWorldContext().World();
-	check(TargetWorld);
-
 	// create new UStaticMesh object
 	EObjectFlags Flags = EObjectFlags::RF_Public | EObjectFlags::RF_Standalone;
 	UStaticMesh* NewStaticMesh = NewObject<UStaticMesh>(AssetPackage, *UniqueAssetName, Flags);
@@ -176,37 +175,47 @@ AActor* CreateMeshAsset(FGeometryCollection& Collection, const TManagedArray<FTr
 	NewStaticMesh->GetBodySetup()->CollisionTraceFlag = ECollisionTraceFlag::CTF_UseComplexAsSimple;
 
 
-	// create new actor
-	FRotator Rotation(0.0f, 0.0f, 0.0f);
-	FActorSpawnParameters SpawnInfo;
-	AStaticMeshActor* NewActor = TargetWorld->SpawnActor<AStaticMeshActor>(FVector::ZeroVector, Rotation, SpawnInfo);
-	NewActor->SetActorLabel(*UniqueAssetName);
-	NewActor->GetStaticMeshComponent()->SetStaticMesh(NewStaticMesh);
+	if (bPlaceInWorld)
+	{
+		UWorld* TargetWorld = GEditor->GetEditorWorldContext().World();
+		check(TargetWorld);
 
-	// if we don't do this, world traces don't hit the mesh
-	NewActor->GetStaticMeshComponent()->RecreatePhysicsState();
+		// create new actor
+		FRotator Rotation(0.0f, 0.0f, 0.0f);
+		FActorSpawnParameters SpawnInfo;
+		AStaticMeshActor* NewActor = TargetWorld->SpawnActor<AStaticMeshActor>(FVector::ZeroVector, Rotation, SpawnInfo);
+		NewActor->SetActorLabel(*UniqueAssetName);
+		NewActor->GetStaticMeshComponent()->SetStaticMesh(NewStaticMesh);
 
-	NewActor->GetStaticMeshComponent()->SetWorldTransform(CollectionToWorld * BonesToCollection);
+		// if we don't do this, world traces don't hit the mesh
+		NewActor->GetStaticMeshComponent()->RecreatePhysicsState();
 
-	for (int MatIdx = 0; MatIdx < Materials.Num(); MatIdx++)
-	{
-		NewActor->GetStaticMeshComponent()->SetMaterial(MatIdx, Materials[MatIdx]);
-	}
+		NewActor->GetStaticMeshComponent()->SetWorldTransform(CollectionToWorld * BonesToCollection);
+
+		for (int MatIdx = 0; MatIdx < Materials.Num(); MatIdx++)
+		{
+			NewActor->GetStaticMeshComponent()->SetMaterial(MatIdx, Materials[MatIdx]);
+		}
 
-	NewActor->MarkComponentsRenderStateDirty();
+		NewActor->MarkComponentsRenderStateDirty();
 
-	GEditor->SelectActor(NewActor, true, false, true, false);
+		if (bSelectNewActors)
+		{
+			GEditor->SelectActor(NewActor, true, false, true, false);
+		}
+	}
 
 	NewStaticMesh->PostEditChange();
 
 	AssetPackage->MarkPackageDirty();
 	FAssetRegistryModule::AssetCreated(AssetPackage);
-
-	return NewActor;
 }
 
 // Convert and save all the requested static mesh assets
-bool ConvertAndSaveMeshes(FGeometryCollection& Collection, const TManagedArray<FTransform>& BoneTransforms, const FTransform& CollectionToWorld, const TConstArrayView<UMaterialInterface*>& Materials, TConstArrayView<int32> Bones, FString ObjectBaseName, const UObject* RelativeToAsset, bool bPromptToSave, bool bSaveCombined, bool bCenterPivots)
+bool ConvertAndSaveMeshes(
+	FGeometryCollection& Collection, const TManagedArray<FTransform>& BoneTransforms, const FTransform& CollectionToWorld,
+	const TConstArrayView<UMaterialInterface*>& Materials, TConstArrayView<int32> Bones, FString ObjectBaseName, 
+	const UObject* RelativeToAsset, bool bPromptToSave, bool bSaveCombined, bool bCenterPivots, bool bPlaceInWorld, bool bSelectNewActors)
 {
 	check(RelativeToAsset);
 
@@ -257,7 +266,10 @@ bool ConvertAndSaveMeshes(FGeometryCollection& Collection, const TManagedArray<F
 	FScopedSlowTask ConvertTask(Bones.Num(), LOCTEXT("StartingConvert", "Converting geometry to static mesh"));
 	ConvertTask.MakeDialog();
 
-	GEditor->SelectNone(false, true, false);
+	if (bSelectNewActors)
+	{
+		GEditor->SelectNone(false, true, false);
+	}
 
 	TArray<UPackage*> SavePackage;
 	if (bSaveCombined)
@@ -266,7 +278,7 @@ bool ConvertAndSaveMeshes(FGeometryCollection& Collection, const TManagedArray<F
 		FString UniqueAssetName;
 		UPackage* AssetPackage = MakeUniquePackage(ObjectBaseName, UniqueAssetName);
 		SavePackage.Add(AssetPackage);
-		CreateMeshAsset(Collection, BoneTransforms, CollectionToWorld, Materials, Bones, AssetPackage, UniqueAssetName, bCenterPivots);
+		CreateMeshAsset(Collection, BoneTransforms, CollectionToWorld, Materials, Bones, AssetPackage, UniqueAssetName, bCenterPivots, bPlaceInWorld, bSelectNewActors);
 	}
 	else
 	{
@@ -284,11 +296,14 @@ bool ConvertAndSaveMeshes(FGeometryCollection& Collection, const TManagedArray<F
 
 			TArrayView<const int32> SingleBoneView(&UseBone, 1);
 
-			CreateMeshAsset(Collection, BoneTransforms, CollectionToWorld, Materials, SingleBoneView, AssetPackage, UniqueAssetName, bCenterPivots);
+			CreateMeshAsset(Collection, BoneTransforms, CollectionToWorld, Materials, SingleBoneView, AssetPackage, UniqueAssetName, bCenterPivots, bPlaceInWorld, bSelectNewActors);
 		}
 	}
 
-	GEditor->NoteSelectionChange(true);
+	if (bSelectNewActors)
+	{
+		GEditor->NoteSelectionChange(true);
+	}
 
 	FEditorFileUtils::PromptForCheckoutAndSave(SavePackage, true, true);
 
@@ -345,7 +360,7 @@ int32 UFractureToolConvert::ExecuteFracture(const FFractureToolContext& Fracture
 		
 		UE::FractureToolConvertInternal::ConvertAndSaveMeshes(Collection, BoneTransforms, FractureContext.GetTransform(), Materials,
 			FractureContext.GetSelection(), BaseName, FractureContext.GetFracturedGeometryCollection(), 
-			ConvertSettings->bPromptForBaseName, !ConvertSettings->bPerBone, ConvertSettings->bCenterPivots);
+			ConvertSettings->bPromptForBaseName, !ConvertSettings->bPerBone, ConvertSettings->bCenterPivots, ConvertSettings->bPlaceInWorld, ConvertSettings->bSelectNewActors);
 	}
 
 	return INDEX_NONE;
diff --git a/Engine/Plugins/Experimental/ChaosEditor/Source/FractureEditor/Private/FractureToolConvert.h b/Engine/Plugins/Experimental/ChaosEditor/Source/FractureEditor/Private/FractureToolConvert.h
--- a/Engine/Plugins/Experimental/ChaosEditor/Source/FractureEditor/Private/FractureToolConvert.h
+++ b/Engine/Plugins/Experimental/ChaosEditor/Source/FractureEditor/Private/FractureToolConvert.h
@@ -32,6 +32,14 @@ public:
 	/** Whether to center the pivot for each mesh, or use the pivot from the geometry collection */
 	UPROPERTY(EditAnywhere, Category = GeometrySettings)
 	bool bCenterPivots = false;
+
+	/** Whether to place new static mesh actors in the level */
+	UPROPERTY(EditAnywhere, Category = ActorSettings)
+	bool bPlaceInWorld = true;
+
+	/** Whether to select new static mesh actors */
+	UPROPERTY(EditAnywhere, Category = ActorSettings, meta = (EditCondition = "bPlaceInWorld"))
+	bool bSelectNewActors = false;
 };
 
 
