Fix for fracture mode auto-cluster failed check()
It's possible (though rare) for a cluster to end up empty, so we should gracefully handle that case

diff --git a/Engine/Plugins/Experimental/ChaosEditor/Source/FractureEditor/Private/FractureToolAutoCluster.cpp b/Engine/Plugins/Experimental/ChaosEditor/Source/FractureEditor/Private/FractureToolAutoCluster.cpp
--- a/Engine/Plugins/Experimental/ChaosEditor/Source/FractureEditor/Private/FractureToolAutoCluster.cpp
+++ b/Engine/Plugins/Experimental/ChaosEditor/Source/FractureEditor/Private/FractureToolAutoCluster.cpp
@@ -95,10 +95,16 @@ void UFractureToolAutoCluster::Execute(TWeakPtr<FFractureEditorModeToolkit> InTo
 				int32 PartitionCount = VoronoiPartition.GetPartitionCount();
 				int32 NewClusterIndexStart = GeometryCollection->AddElements(PartitionCount, FGeometryCollection::TransformGroup);
 
+				bool bHasEmptyClusters = false;
+
 				for (int32 Index = 0; Index < PartitionCount; ++Index)
 				{
 						
 					TArray<int32> NewCluster = VoronoiPartition.GetPartition(Index);
+					if (NewCluster.Num() == 0)
+					{
+						bHasEmptyClusters = true;
+					}
 
 					int32 NewClusterIndex = NewClusterIndexStart + Index;
 					GeometryCollection->Parent[NewClusterIndex] = ClusterIndex;
@@ -109,6 +115,10 @@ void UFractureToolAutoCluster::Execute(TWeakPtr<FFractureEditorModeToolkit> InTo
 					GeometryCollection->Transform[NewClusterIndex] = FTransform::Identity;
 					GeometryCollectionAlgo::ParentTransforms(GeometryCollection, NewClusterIndex, NewCluster);
 				}
+				if (bHasEmptyClusters)
+				{
+					FGeometryCollectionClusteringUtility::RemoveDanglingClusters(GeometryCollection);
+				}
 				FGeometryCollectionClusteringUtility::UpdateHierarchyLevelOfChildren(GeometryCollection, ClusterIndex);
 				FGeometryCollectionClusteringUtility::RecursivelyUpdateChildBoneNames(ClusterIndex, GeometryCollection->Children, GeometryCollection->BoneName);
 				FGeometryCollectionClusteringUtility::ValidateResults(GeometryCollection);
@@ -357,7 +367,10 @@ bool FVoronoiPartitioner::Refine()
 	// Recalculate partition centers
 	for (int32 PartitionIndex = 0; PartitionIndex < PartitionCount; ++PartitionIndex)
 	{
-		PartitionCenters[PartitionIndex] = FVector(0.0f, 0.0f, 0.0f);
+		if (PartitionSize[PartitionIndex] > 0)
+		{
+			PartitionCenters[PartitionIndex] = FVector(0.0f, 0.0f, 0.0f);
+		}
 	}
 
 	for (int32 Index = 0; Index < Partitions.Num(); ++Index)
@@ -367,8 +380,10 @@ bool FVoronoiPartitioner::Refine()
 
 	for (int32 PartitionIndex = 0; PartitionIndex < PartitionCount; ++PartitionIndex)
 	{
-		check(PartitionSize[PartitionIndex] > 0);
-		PartitionCenters[PartitionIndex] /= float(PartitionSize[PartitionIndex]);
+		if (PartitionSize[PartitionIndex] > 0)
+		{
+			PartitionCenters[PartitionIndex] /= float(PartitionSize[PartitionIndex]);
+		}
 	}
 
 	return true;
