fix empty (but present) proximity attribute on new/reset of geometry collection

diff --git a/Engine/Plugins/Experimental/ChaosEditor/Source/FractureEditor/Private/FractureToolGenerators.cpp b/Engine/Plugins/Experimental/ChaosEditor/Source/FractureEditor/Private/FractureToolGenerators.cpp
--- a/Engine/Plugins/Experimental/ChaosEditor/Source/FractureEditor/Private/FractureToolGenerators.cpp
+++ b/Engine/Plugins/Experimental/ChaosEditor/Source/FractureEditor/Private/FractureToolGenerators.cpp
@@ -19,6 +19,7 @@
 #include "GeometryCollection/GeometryCollectionClusteringUtility.h"
 #include "GeometryCollection/GeometryCollectionConversion.h"
 #include "GeometryCollection/GeometryCollectionUtility.h"
+#include "GeometryCollection/GeometryCollectionProximityUtility.h"
 #include "FractureToolContext.h"
 
 
@@ -258,6 +259,10 @@ AGeometryCollectionActor* UFractureToolGenerateAsset::ConvertActorsToGeometryCol
 	// Add and initialize guids
 	::GeometryCollection::GenerateTemporaryGuids(FracturedGeometryCollection->GetGeometryCollection().Get(), 0 , true);
 
+	// Update proximity graph
+	FGeometryCollectionProximityUtility ProximityUtility(FracturedGeometryCollection->GetGeometryCollection().Get());
+	ProximityUtility.UpdateProximity();
+
 	return NewActor;
 }
 
@@ -438,7 +443,9 @@ void UFractureToolResetAsset::Execute(TWeakPtr<FFractureEditorModeToolkit> InToo
 					}
 				}
 				
-
+				// Update proximity graph
+				FGeometryCollectionProximityUtility ProximityUtility(GeometryCollection);
+				ProximityUtility.UpdateProximity();
 
 				FGeometryCollectionClusteringUtility::UpdateHierarchyLevelOfChildren(GeometryCollection, -1);
 				AddSingleRootNodeIfRequired(GeometryCollectionObject);
