fix issues caused by attempting to fracture/cluster/etc while sim/PIE is running (by preventing these actions while sim is running)

diff --git a/Engine/Plugins/Experimental/ChaosEditor/Source/FractureEditor/Private/FractureEditorModeToolkit.cpp b/Engine/Plugins/Experimental/ChaosEditor/Source/FractureEditor/Private/FractureEditorModeToolkit.cpp
--- a/Engine/Plugins/Experimental/ChaosEditor/Source/FractureEditor/Private/FractureEditorModeToolkit.cpp
+++ b/Engine/Plugins/Experimental/ChaosEditor/Source/FractureEditor/Private/FractureEditorModeToolkit.cpp
@@ -209,6 +209,11 @@ void FFractureEditorModeToolkit::Init(const TSharedPtr<IToolkitHost>& InitToolki
 	FLevelEditorModule& LevelEditorModule = FModuleManager::LoadModuleChecked<FLevelEditorModule>("LevelEditor");
 	LevelEditorModule.OnMapChanged().AddRaw(this, &FFractureEditorModeToolkit::HandleMapChanged);
 
+	BeginPIEDelegateHandle = FEditorDelegates::BeginPIE.AddLambda([this](bool bSimulating)
+	{
+		SetActiveTool(nullptr);
+	});
+
 	FDetailsViewArgs DetailsViewArgs;
 	DetailsViewArgs.bAllowSearch = false;
 	DetailsViewArgs.NameAreaSettings = FDetailsViewArgs::HideNameArea;
@@ -870,7 +875,7 @@ void FFractureEditorModeToolkit::BindCommands()
 				ToolkitCommands->MapAction(
 					FractureTool->GetUICommandInfo(),
 					FExecuteAction::CreateSP(this, &FFractureEditorModeToolkit::SetActiveTool, FractureTool),
-					FCanExecuteAction(),
+					FCanExecuteAction::CreateSP(this, &FFractureEditorModeToolkit::CanSetModalTool, FractureTool),
 					FIsActionChecked::CreateSP(this, &FFractureEditorModeToolkit::IsActiveTool, FractureTool)
 				);
 			}
@@ -1257,6 +1262,12 @@ void FFractureEditorModeToolkit::InvalidateHitProxies()
 
 bool FFractureEditorModeToolkit::CanExecuteAction(UFractureActionTool* InActionTool) const
 {
+	// Disallow fracture actions when playing in editor or simulating.
+	if (GEditor->PlayWorld || GIsPlayInEditorWorld)
+	{
+		return false;
+	}
+
 	if (InActionTool)
 	{
 		return InActionTool->CanExecute();
@@ -1276,6 +1287,17 @@ void FFractureEditorModeToolkit::ShutdownActiveTool()
 	}
 }
 
+bool FFractureEditorModeToolkit::CanSetModalTool(UFractureModalTool* InActiveTool) const
+{
+	// Disallow fracture modal tools when playing in editor or simulating.
+	if (GEditor->PlayWorld || GIsPlayInEditorWorld)
+	{
+		return false;
+	}
+
+	return true;
+}
+
 void FFractureEditorModeToolkit::SetActiveTool(UFractureModalTool* InActiveTool)
 {
 	ShutdownActiveTool();
@@ -1311,6 +1333,8 @@ void FFractureEditorModeToolkit::InvalidateCachedDetailPanelState(UObject* Chang
 void FFractureEditorModeToolkit::Shutdown()
 {
 	ShutdownActiveTool();
+
+	FEditorDelegates::BeginPIE.Remove(BeginPIEDelegateHandle);
 }
 
 
@@ -1423,6 +1447,10 @@ FReply FFractureEditorModeToolkit::OnModalClicked()
 
 bool FFractureEditorModeToolkit::CanExecuteModal() const
 {
+	if (GEditor->PlayWorld || GIsPlayInEditorWorld)
+	{
+		return false;
+	}
 
 	if (!IsSelectedActorsInEditorWorld())
 	{
diff --git a/Engine/Plugins/Experimental/ChaosEditor/Source/FractureEditor/Public/FractureEditorModeToolkit.h b/Engine/Plugins/Experimental/ChaosEditor/Source/FractureEditor/Public/FractureEditorModeToolkit.h
--- a/Engine/Plugins/Experimental/ChaosEditor/Source/FractureEditor/Public/FractureEditorModeToolkit.h
+++ b/Engine/Plugins/Experimental/ChaosEditor/Source/FractureEditor/Public/FractureEditorModeToolkit.h
@@ -99,6 +99,7 @@ public:
 	void ExecuteAction(UFractureActionTool* InActionTool);
 	bool CanExecuteAction(UFractureActionTool* InActionTool) const;
 
+	bool CanSetModalTool(UFractureModalTool* InActiveTool) const;
 	void SetActiveTool(UFractureModalTool* InActiveTool);
 	UFractureModalTool* GetActiveTool() const;
 	bool IsActiveTool(UFractureModalTool* InActiveTool);
@@ -202,6 +203,9 @@ private:
 private:
 	UFractureModalTool* ActiveTool;
 
+	// called when PIE is about to start, shuts down active tools
+	FDelegateHandle BeginPIEDelegateHandle;
+
 	TSharedPtr<IDetailsView> DetailsView;
 	TSharedPtr<IDetailsView> ViewSettingsDetailsView;
 	TSharedPtr<IDetailsView> HistogramDetailsView;
