fix fracture convex tool re-generating convexes with unsaved UI settings instead of the saved settings on the collection when user presses the 'clear custom' button

diff --git a/Engine/Plugins/Experimental/ChaosEditor/Source/FractureEditor/Private/FractureToolConvex.cpp b/Engine/Plugins/Experimental/ChaosEditor/Source/FractureEditor/Private/FractureToolConvex.cpp
--- a/Engine/Plugins/Experimental/ChaosEditor/Source/FractureEditor/Private/FractureToolConvex.cpp
+++ b/Engine/Plugins/Experimental/ChaosEditor/Source/FractureEditor/Private/FractureToolConvex.cpp
@@ -48,8 +48,13 @@ void UFractureToolConvex::DeleteConvexFromSelected()
 			continue;
 		}
 
+		// Note: Must make sure the custom flags attribute is added *before* the CollectionEdit calls Modify() on the collection, because serialization is weird for
+		// the geometry collection: undo will load the saved attributes back, but not clear the attributes that didn't exist before
+		// TODO: Consider making the serialization clear the attributes that didn't exist before, then move this line after the CollectionEdit line
 		TManagedArray<int32>& HasCustomConvex = *FGeometryCollectionConvexUtility::GetCustomConvexFlags(&Collection, true);
 
+		FGeometryCollectionEdit CollectionEdit(FractureContext.GetGeometryCollectionComponent(), GeometryCollection::EEditUpdate::Rest, true /*bShapeIsUnchanged -- use to skip re-generating convexes, because we directly update them */);
+
 		TArray<int32> TransformsToClear;
 		for (int32 TransformIdx : FractureContext.GetSelection())
 		{
@@ -72,8 +77,6 @@ void UFractureToolConvex::PromoteChildren()
 
 	for (FFractureToolContext& FractureContext : FractureContexts)
 	{
-		UFractureActionTool::FModifyContextScope ModifyScope(this, &FractureContext);
-
 		FGeometryCollection& Collection = *FractureContext.GetGeometryCollection();
 
 		if (!Collection.HasAttribute("ConvexHull", "Convex") ||
@@ -81,6 +84,13 @@ void UFractureToolConvex::PromoteChildren()
 		{
 			continue;
 		}
+		
+		// Note: Must make sure the custom flags attribute is added *before* the CollectionEdit calls Modify() on the collection, because serialization is weird for
+		// the geometry collection: undo will load the saved attributes back, but not clear the attributes that didn't exist before
+		// TODO: Consider making the serialization clear the attributes that didn't exist before, then move this line after the CollectionEdit line
+		FGeometryCollectionConvexUtility::GetCustomConvexFlags(FractureContext.GetGeometryCollection().Get(), true);
+
+		FGeometryCollectionEdit CollectionEdit(FractureContext.GetGeometryCollectionComponent(), GeometryCollection::EEditUpdate::Rest, true /*bShapeIsUnchanged -- use to skip re-generating convexes, because we directly update them */);
 
 		TArray<int32> SelectedTransforms = FractureContext.GetSelection();
 		FGeometryCollectionConvexUtility::CopyChildConvexes(&Collection, SelectedTransforms, &Collection, SelectedTransforms, false);
@@ -104,7 +114,7 @@ void UFractureToolConvex::ClearCustomConvex()
 			continue;
 		}
 
-		UFractureActionTool::FModifyContextScope ModifyScope(this, &FractureContext);
+		FGeometryCollectionEdit CollectionEdit(FractureContext.GetGeometryCollectionComponent(), GeometryCollection::EEditUpdate::Rest);
 
 		FGeometryCollection& Collection = *FractureContext.GetGeometryCollection();
 		for (int32 TransformIdx : FractureContext.GetSelection())
@@ -130,8 +140,6 @@ void UFractureToolConvex::ClearCustomConvex()
 			{
 				Collection.RemoveAttribute("HasCustomConvex", FTransformCollection::TransformGroup);
 			}
-
-			AutoComputeConvex(FractureContext);
 		}
 	}
 
@@ -277,18 +285,6 @@ TArray<FFractureToolContext> UFractureToolConvex::GetFractureToolContexts() cons
 	return Contexts;
 }
 
-void UFractureToolConvex::AutoComputeConvex(const FFractureToolContext& FractureContext)
-{
-	if (FractureContext.GetGeometryCollection().IsValid())
-	{
-		FGeometryCollection& Collection = *FractureContext.GetGeometryCollection();
-		FGeometryCollectionProximityUtility ProximityUtility(&Collection);
-		ProximityUtility.UpdateProximity();
-		FGeometryCollectionConvexUtility::CreateNonOverlappingConvexHullData(&Collection, ConvexSettings->FractionAllowRemove, ConvexSettings->SimplificationDistanceThreshold, ConvexSettings->CanExceedFraction);
-	}
-}
-
-
 int32 UFractureToolConvex::ExecuteFracture(const FFractureToolContext& FractureContext)
 {
 	if (FractureContext.GetGeometryCollection().IsValid())
diff --git a/Engine/Plugins/Experimental/ChaosEditor/Source/FractureEditor/Private/FractureToolConvex.h b/Engine/Plugins/Experimental/ChaosEditor/Source/FractureEditor/Private/FractureToolConvex.h
--- a/Engine/Plugins/Experimental/ChaosEditor/Source/FractureEditor/Private/FractureToolConvex.h
+++ b/Engine/Plugins/Experimental/ChaosEditor/Source/FractureEditor/Private/FractureToolConvex.h
@@ -84,7 +84,6 @@ public:
 	void DeleteConvexFromSelected();
 	void PromoteChildren();
 	void ClearCustomConvex();
-	void AutoComputeConvex(const FFractureToolContext& FractureContext);
 
 
 protected:
