fix fracture mode user facing text for accuracy and clarity

diff --git a/Engine/Plugins/Experimental/ChaosEditor/Source/FractureEditor/Private/FractureToolAutoCluster.h b/Engine/Plugins/Experimental/ChaosEditor/Source/FractureEditor/Private/FractureToolAutoCluster.h
--- a/Engine/Plugins/Experimental/ChaosEditor/Source/FractureEditor/Private/FractureToolAutoCluster.h
+++ b/Engine/Plugins/Experimental/ChaosEditor/Source/FractureEditor/Private/FractureToolAutoCluster.h
@@ -6,7 +6,7 @@
 
 #include "FractureToolAutoCluster.generated.h"
 
-
+// Note: Only Voronoi-based auto-clustering is currently supported
 UENUM()
 enum class EFractureAutoClusterMode : uint8
 {
@@ -37,9 +37,11 @@ public:
 	UPROPERTY(meta = (DeprecatedProperty, DeprecationMessage = "Simplified interface now only supports Voronoi clustering."))
 	EFractureAutoClusterMode AutoClusterMode_DEPRECATED;
 
+	/** Use a Voronoi diagram with this many Voronoi sites as a guide for deciding cluster boundaries */
 	UPROPERTY(EditAnywhere, Category = AutoCluster, meta = (DisplayName = "Cluster Sites", UIMin = "1", UIMax = "5000", ClampMin = "1"))
-	uint32 SiteCount=1;
+	uint32 SiteCount=10;
 
+	/** If true, bones will only be added to the same cluster if they are physically connected (either directly, or via other bones in the same cluster) */
 	UPROPERTY(EditAnywhere, Category = AutoCluster, meta = (DisplayName = "Enforce Cluster Connectivity"))
 	bool bEnforceConnectivity=true;
 };
diff --git a/Engine/Plugins/Experimental/ChaosEditor/Source/FractureEditor/Private/FractureToolBrick.cpp b/Engine/Plugins/Experimental/ChaosEditor/Source/FractureEditor/Private/FractureToolBrick.cpp
--- a/Engine/Plugins/Experimental/ChaosEditor/Source/FractureEditor/Private/FractureToolBrick.cpp
+++ b/Engine/Plugins/Experimental/ChaosEditor/Source/FractureEditor/Private/FractureToolBrick.cpp
@@ -40,16 +40,16 @@ FSlateIcon UFractureToolBrick::GetToolIcon() const
 
 void UFractureToolBrick::RegisterUICommand( FFractureEditorCommands* BindingContext ) 
 {
-	UI_COMMAND_EXT( BindingContext, UICommandInfo, "Brick", "Brick", "Fracture with a customizable brick pattern.", EUserInterfaceActionType::ToggleButton, FInputChord() );
+	UI_COMMAND_EXT( BindingContext, UICommandInfo, "Brick", "Brick", "Fracture with a customizable brick pattern. Note: Currently only supports fracturing with at least some (non-zero) Grout.", EUserInterfaceActionType::ToggleButton, FInputChord() );
 	BindingContext->Brick = UICommandInfo;
 }
 
 TArray<UObject*> UFractureToolBrick::GetSettingsObjects() const 
 { 
-	TArray<UObject*> Settings; 
+	TArray<UObject*> Settings;
+	Settings.Add(BrickSettings);
 	Settings.Add(CutterSettings);
 	Settings.Add(CollisionSettings);
-	Settings.Add(BrickSettings);
 	return Settings;
 }
 
@@ -417,7 +417,7 @@ int32 UFractureToolBrick::ExecuteFracture(const FFractureToolContext& FractureCo
 		// Proximity is invalidated.
 		ClearProximity(FractureContext.GetGeometryCollection().Get());
 
-		return CutMultipleWithPlanarCells(VoronoiPlanarCells, *FractureContext.GetGeometryCollection(), FractureContext.GetSelection(), 0, CollisionSettings->PointSpacing, FractureContext.GetSeed(), FractureContext.GetTransform());
+		return CutMultipleWithPlanarCells(VoronoiPlanarCells, *FractureContext.GetGeometryCollection(), FractureContext.GetSelection(), 0, CollisionSettings->GetPointSpacing(), FractureContext.GetSeed(), FractureContext.GetTransform());
 	}
 
 	return INDEX_NONE;
diff --git a/Engine/Plugins/Experimental/ChaosEditor/Source/FractureEditor/Private/FractureToolBrick.h b/Engine/Plugins/Experimental/ChaosEditor/Source/FractureEditor/Private/FractureToolBrick.h
--- a/Engine/Plugins/Experimental/ChaosEditor/Source/FractureEditor/Private/FractureToolBrick.h
+++ b/Engine/Plugins/Experimental/ChaosEditor/Source/FractureEditor/Private/FractureToolBrick.h
@@ -9,18 +9,19 @@
 
 class FFractureToolContext;
 
-/** Brick Projection Directions*/
-UENUM()
-enum class EFractureBrickProjection: uint8
-{
-	X UMETA(DisplayName = "X"),
-
-	Y UMETA(DisplayName = "Y"),
-
-	Z UMETA(DisplayName = "Z")
-};
-
-/** Brick Projection Directions*/
+// TODO: this enum was unused in the brick generation code; consider whether we want to keep + use it or not
+///** Brick Projection Directions*/
+//UENUM()
+//enum class EFractureBrickProjection: uint8
+//{
+//	X UMETA(DisplayName = "X"),
+//
+//	Y UMETA(DisplayName = "Y"),
+//
+//	Z UMETA(DisplayName = "Z")
+//};
+
+/** Brick Stacking Pattern */
 UENUM()
 enum class EFractureBrickBond : uint8
 {
@@ -45,34 +46,32 @@ public:
 	UFractureBrickSettings(const FObjectInitializer& ObjInit)
 		: Super(ObjInit)
 		, Bond(EFractureBrickBond::Stretcher)
-		, Forward(EFractureBrickProjection::X)
-		, Up(EFractureBrickProjection::Z)
 		, BrickLength(194.f)
 		, BrickHeight(57.f)
 		, BrickDepth(92.f)
 	{}
 
-	/** Forward Direction to project brick pattern. */
+	/** The brick bond pattern defines how the bricks are arranged */
 	UPROPERTY(EditAnywhere, Category = Brick)
 	EFractureBrickBond Bond;
 
-	/** Forward Direction to project brick pattern. */
-	UPROPERTY(EditAnywhere, Category = Brick)
-	EFractureBrickProjection Forward;
-
-	/** Up Direction for vertical brick slices. */
-	UPROPERTY(EditAnywhere, Category = Brick)
-	EFractureBrickProjection Up;
+	// Note: These Forward and Up parameters were unused in the actual brick generation; TODO: consider if we want to actually keep + use these parameters or not
+	///** Forward Direction to project brick pattern. */
+	//UPROPERTY(EditAnywhere, Category = Brick)
+	//EFractureBrickProjection Forward = EFractureBrickProjection::X;
+	///** Up Direction for vertical brick slices. */
+	//UPROPERTY(EditAnywhere, Category = Brick)
+	//EFractureBrickProjection Up = EFractureBrickProjection::Z;
 
-	/** Brick Length */
+	/** Brick length (in cm) */
 	UPROPERTY(EditAnywhere, Category = Brick, meta = (UIMin = "0.1", UIMax = "500.0", ClampMin = "0.001"))
 	float BrickLength;
 
-	/** Brick Height */
+	/** Brick height (in cm) */
 	UPROPERTY(EditAnywhere, Category = Brick, meta = (UIMin = "0.1", UIMax = "500.0", ClampMin = "0.001"))
 	float BrickHeight;
 
-	/** Brick Depth */
+	/** Brick depth (in cm) */
 	UPROPERTY(EditAnywhere, Category = Brick, meta = (UIMin = "0.1", UIMax = "500.0", ClampMin = "0.001"))
 	float BrickDepth;
 };
diff --git a/Engine/Plugins/Experimental/ChaosEditor/Source/FractureEditor/Private/FractureToolClusterCutter.cpp b/Engine/Plugins/Experimental/ChaosEditor/Source/FractureEditor/Private/FractureToolClusterCutter.cpp
--- a/Engine/Plugins/Experimental/ChaosEditor/Source/FractureEditor/Private/FractureToolClusterCutter.cpp
+++ b/Engine/Plugins/Experimental/ChaosEditor/Source/FractureEditor/Private/FractureToolClusterCutter.cpp
@@ -39,9 +39,9 @@ void UFractureToolClusterCutter::RegisterUICommand( FFractureEditorCommands* Bin
 TArray<UObject*> UFractureToolClusterCutter::GetSettingsObjects() const
 {
 	TArray<UObject*> ReturnSettings;
+	ReturnSettings.Add(ClusterSettings);
 	ReturnSettings.Add(CutterSettings);
 	ReturnSettings.Add(CollisionSettings);
-	ReturnSettings.Add(ClusterSettings);
 	return ReturnSettings;
 }
 
@@ -69,7 +69,7 @@ void UFractureToolClusterCutter::GenerateVoronoiSites(const FFractureToolContext
 		{
 			FVector V(RandStream.VRand());
 			V.Normalize();
-			V *= ClusterSettings->ClusterRadius + (RandStream.FRandRange(ClusterSettings->ClusterRadiusPercentageMin, ClusterSettings->ClusterRadiusPercentageMax) * Bounds.GetExtent().GetAbsMax());
+			V *= ClusterSettings->ClusterRadiusOffset + (RandStream.FRandRange(ClusterSettings->ClusterRadiusFractionMin, ClusterSettings->ClusterRadiusFractionMax) * Bounds.GetExtent().GetAbsMax());
 			V += CenterSites[kk];
 			Sites.Emplace(V);
 		}
diff --git a/Engine/Plugins/Experimental/ChaosEditor/Source/FractureEditor/Private/FractureToolClusterCutter.h b/Engine/Plugins/Experimental/ChaosEditor/Source/FractureEditor/Private/FractureToolClusterCutter.h
--- a/Engine/Plugins/Experimental/ChaosEditor/Source/FractureEditor/Private/FractureToolClusterCutter.h
+++ b/Engine/Plugins/Experimental/ChaosEditor/Source/FractureEditor/Private/FractureToolClusterCutter.h
@@ -21,37 +21,44 @@ public:
 		, NumberClustersMax(8)
 		, SitesPerClusterMin(2)
 		, SitesPerClusterMax(30)
-		, ClusterRadiusPercentageMin(0.1)
-		, ClusterRadiusPercentageMax(0.2)
-		, ClusterRadius(0.0f)
+		, ClusterRadiusFractionMin(0.1)
+		, ClusterRadiusFractionMax(0.2)
+		, ClusterRadiusOffset(0.0f)
 	{}
 
-	/** Number of Clusters - Cluster Voronoi Method */
-	UPROPERTY(EditAnywhere, Category = ClusterVoronoi, meta = (DisplayName = "Minimum Cluster Sites", UIMin = "1", UIMax = "200", ClampMin = "1"))
+	/** Minimum number of clusters of Voronoi sites to create. The amount of clusters created will be chosen at random between Min and Max */
+	UPROPERTY(EditAnywhere, Category = ClusterVoronoi, meta = (DisplayName = "Min Num Clusters", UIMin = "1", UIMax = "200", ClampMin = "1"))
 	int32 NumberClustersMin;
 
-	/** Number of Clusters - Cluster Voronoi Method */
-	UPROPERTY(EditAnywhere, Category = ClusterVoronoi, meta = (DisplayName = "Maximum Cluster Sites", UIMin = "1", UIMax = "200", ClampMin = "1"))
+	/** Maximum number of clusters of Voronoi sites to create. The amount of clusters created will be chosen at random between Min and Max */
+	UPROPERTY(EditAnywhere, Category = ClusterVoronoi, meta = (DisplayName = "Max Num Clusters", UIMin = "1", UIMax = "200", ClampMin = "1"))
 	int32 NumberClustersMax;
 
-	/** Sites per # of Clusters - Cluster Voronoi Method */
-	UPROPERTY(EditAnywhere, Category = ClusterVoronoi, meta = (DisplayName = "Minimum Sites Per Cluster", UIMin = "0", UIMax = "200", ClampMin = "0"))
+	/** Minimum number of Voronoi sites per cluster. The amount of sites in each cluster will be chosen at random between Min and Max */
+	UPROPERTY(EditAnywhere, Category = ClusterVoronoi, meta = (DisplayName = "Min Sites Per Cluster", UIMin = "0", UIMax = "200", ClampMin = "0"))
 	int32 SitesPerClusterMin;
 
-	UPROPERTY(EditAnywhere, Category = ClusterVoronoi, meta = (DisplayName = "Maximum Sites Per Cluster", UIMin = "0", UIMax = "200", ClampMin = "0"))
+	/** Maximum number of Voronoi sites per cluster. The amount of sites in each cluster will be chosen at random between Min and Max */
+	UPROPERTY(EditAnywhere, Category = ClusterVoronoi, meta = (DisplayName = "Max Sites Per Cluster", UIMin = "0", UIMax = "200", ClampMin = "0"))
 	int32 SitesPerClusterMax;
 
-	/** Cluster's Radius - Cluster Voronoi Method */
-	UPROPERTY(EditAnywhere, Category = ClusterVoronoi, meta = (DisplayName = "Minimum distance from center as part of bounds max extent", UIMin = "0.0", UIMax = "1.0"))
-	float ClusterRadiusPercentageMin;
-
-	/** Cluster's Radius - Cluster Voronoi Method */
-	UPROPERTY(EditAnywhere, Category = ClusterVoronoi, meta = (DisplayName = "Maximum distance from center as part of bounds max extent", UIMin = "0.0", UIMax = "1.0"))
-	float ClusterRadiusPercentageMax;
-
-	/** Cluster's Radius - Cluster Voronoi Method */
+	/**
+	 * Minimum cluster radius (as fraction of the overall Voronoi diagram size). Cluster Radius Offset will be added to this
+	 * Each Voronoi site will be placed at least this far (plus the Cluster Radius Offset) from its cluster center
+	 */
+	UPROPERTY(EditAnywhere, Category = ClusterVoronoi, meta = (DisplayName = "Min Dist from Center (as Frac of Bounds)", UIMin = "0.0", UIMax = "1.0"))
+	float ClusterRadiusFractionMin;
+
+	/**
+	 * Maximum cluster radius (as fraction of the overall Voronoi diagram size). Cluster Radius Offset will be added to this
+	 * Each Voronoi site will be placed at most this far (plus the Cluster Radius Offset) from its cluster center
+	 */
+	UPROPERTY(EditAnywhere, Category = ClusterVoronoi, meta = (DisplayName = "Max Dist from Center (as Frac of Bounds)", UIMin = "0.0", UIMax = "1.0"))
+	float ClusterRadiusFractionMax;
+
+	/** Cluster radius offset (in cm). This offset will be added to the 'Min/Max Dist from Center' distance */
 	UPROPERTY(EditAnywhere, Category = ClusterVoronoi)
-	float ClusterRadius;
+	float ClusterRadiusOffset;
 };
 
 
diff --git a/Engine/Plugins/Experimental/ChaosEditor/Source/FractureEditor/Private/FractureToolClusterMagnet.cpp b/Engine/Plugins/Experimental/ChaosEditor/Source/FractureEditor/Private/FractureToolClusterMagnet.cpp
--- a/Engine/Plugins/Experimental/ChaosEditor/Source/FractureEditor/Private/FractureToolClusterMagnet.cpp
+++ b/Engine/Plugins/Experimental/ChaosEditor/Source/FractureEditor/Private/FractureToolClusterMagnet.cpp
@@ -9,7 +9,6 @@
 
 #include "Chaos/TriangleMesh.h"
 #include "PhysicsProxy/GeometryCollectionPhysicsProxy.h"
-#include "Chaos/MassProperties.h"
 #include "Chaos/Particles.h"
 #include "Chaos/Vector.h"
 
@@ -40,7 +39,7 @@ FText UFractureToolClusterMagnet::GetDisplayText() const
 
 FText UFractureToolClusterMagnet::GetTooltipText() const
 {
-	return FText(NSLOCTEXT("Fracture", "FractureToolClusterMagnetToolTip", "Builds clusters at local level by collecting bones adjacent to clusters or bones with highest mass."));
+	return FText(NSLOCTEXT("Fracture", "FractureToolClusterMagnetToolTip", "Builds clusters by grouping the selected bones with their adjacent, neighboring bones. Can iteratively expand to a larger set of neighbors-of-neighbors."));
 }
 
 FText UFractureToolClusterMagnet::GetApplyText() const
@@ -63,7 +62,7 @@ TArray<UObject*> UFractureToolClusterMagnet::GetSettingsObjects() const
 
 void UFractureToolClusterMagnet::RegisterUICommand(FFractureEditorCommands* BindingContext)
 {
-	UI_COMMAND_EXT(BindingContext, UICommandInfo, "ClusterMagnet", "Magnet", "Builds clusters at local level by collecting bones adjacent to clusters or bones with highest mass.", EUserInterfaceActionType::ToggleButton, FInputChord());
+	UI_COMMAND_EXT(BindingContext, UICommandInfo, "ClusterMagnet", "Magnet", "Builds clusters by grouping the selected bones with their adjacent, neighboring bones. Can iteratively expand to a larger set of neighbors-of-neighbors.", EUserInterfaceActionType::ToggleButton, FInputChord());
 	BindingContext->ClusterMagnet = UICommandInfo;
 }
 
@@ -219,8 +218,6 @@ void UFractureToolClusterMagnet::SeparateClusterMagnets(
 	TArray<FClusterMagnet>& OutClusterMagnets,
 	TSet<int32>& OutRemainingPool) const
 {
-	// Push any top nodes over the mass threshold into cluster magnets.
-	
 	OutClusterMagnets.Reserve(TopNodes.Num());
 	OutRemainingPool.Reserve(TopNodes.Num());
 	
diff --git a/Engine/Plugins/Experimental/ChaosEditor/Source/FractureEditor/Private/FractureToolClusterMagnet.h b/Engine/Plugins/Experimental/ChaosEditor/Source/FractureEditor/Private/FractureToolClusterMagnet.h
--- a/Engine/Plugins/Experimental/ChaosEditor/Source/FractureEditor/Private/FractureToolClusterMagnet.h
+++ b/Engine/Plugins/Experimental/ChaosEditor/Source/FractureEditor/Private/FractureToolClusterMagnet.h
@@ -32,6 +32,7 @@ public:
 		, Iterations(1)
 	{}
 
+	/** How widely to search for neighboring bones to add to the new clusters: When Iterations is 1, only direct neighbors of the selected bones can be added to the new clusters; increasing Iterations allows additional layers of neighbors-of-neighbors to be added */
 	UPROPERTY(EditAnywhere, Category = ClusterMagnet, meta = (DisplayName = "Iterations", UIMin = "1", UIMax = "100", ClampMin = "1", ClampMax = "500"))
 	uint32 Iterations;
 
diff --git a/Engine/Plugins/Experimental/ChaosEditor/Source/FractureEditor/Private/FractureToolCustomVoronoi.h b/Engine/Plugins/Experimental/ChaosEditor/Source/FractureEditor/Private/FractureToolCustomVoronoi.h
--- a/Engine/Plugins/Experimental/ChaosEditor/Source/FractureEditor/Private/FractureToolCustomVoronoi.h
+++ b/Engine/Plugins/Experimental/ChaosEditor/Source/FractureEditor/Private/FractureToolCustomVoronoi.h
@@ -12,7 +12,7 @@ class FFractureToolContext;
 UENUM()
 enum class EVoronoiPattern
 {
-	// Add a single site centered on the gizmo
+	// Add a single site centered on the gizmo (or multiple if Variability is > 0)
 	Centered,
 	// Add uniform-random sites
 	Uniform,
@@ -54,7 +54,7 @@ public:
 	UPROPERTY(EditAnywhere, Category = LiveVoronoiSites, meta = (EditCondition = "VoronoiPattern == EVoronoiPattern::MeshVertices || VoronoiPattern == EVoronoiPattern::SelectedBones"))
 	float NormalOffset = 0.0f;
 
-	/** Randomness of sites distribution */
+	/** Amount to randomly displace each Voronoi site (in cm) */
 	UPROPERTY(EditAnywhere, Category = LiveVoronoiSites, meta = (DisplayName = "Variability", UIMin = "0.0", ClampMin = "0.0"))
 	float Variability = 0.0f;
 
diff --git a/Engine/Plugins/Experimental/ChaosEditor/Source/FractureEditor/Private/FractureToolCutter.cpp b/Engine/Plugins/Experimental/ChaosEditor/Source/FractureEditor/Private/FractureToolCutter.cpp
--- a/Engine/Plugins/Experimental/ChaosEditor/Source/FractureEditor/Private/FractureToolCutter.cpp
+++ b/Engine/Plugins/Experimental/ChaosEditor/Source/FractureEditor/Private/FractureToolCutter.cpp
@@ -118,7 +118,7 @@ void UFractureCutterSettings::TransferNoiseSettings(FNoiseSettings& NoiseSetting
 	NoiseSettingsOut.Lacunarity = Lacunarity;
 	NoiseSettingsOut.Persistence = Persistence;
 	NoiseSettingsOut.Octaves = OctaveNumber;
-	NoiseSettingsOut.PointSpacing = SurfaceResolution;
+	NoiseSettingsOut.PointSpacing = PointSpacing;
 }
 
 UFractureToolCutterBase::UFractureToolCutterBase(const FObjectInitializer& ObjInit)
@@ -368,7 +368,7 @@ int32 UFractureToolVoronoiCutterBase::ExecuteFracture(const FFractureToolContext
 		// Proximity is invalidated.
 		ClearProximity(FractureContext.GetGeometryCollection().Get());
 
-		return CutMultipleWithPlanarCells(VoronoiPlanarCells, *(FractureContext.GetGeometryCollection()), FractureContext.GetSelection(), CutterSettings->Grout, CollisionSettings->PointSpacing, FractureContext.GetSeed(), FractureContext.GetTransform());
+		return CutMultipleWithPlanarCells(VoronoiPlanarCells, *(FractureContext.GetGeometryCollection()), FractureContext.GetSelection(), CutterSettings->Grout, CollisionSettings->GetPointSpacing(), FractureContext.GetSeed(), FractureContext.GetTransform());
 	}
 
 	return INDEX_NONE;
diff --git a/Engine/Plugins/Experimental/ChaosEditor/Source/FractureEditor/Private/FractureToolCutter.h b/Engine/Plugins/Experimental/ChaosEditor/Source/FractureEditor/Private/FractureToolCutter.h
--- a/Engine/Plugins/Experimental/ChaosEditor/Source/FractureEditor/Private/FractureToolCutter.h
+++ b/Engine/Plugins/Experimental/ChaosEditor/Source/FractureEditor/Private/FractureToolCutter.h
@@ -31,19 +31,19 @@ public:
 		, bDrawDiagram(true)
 		, Amplitude(0.0f)
 		, Frequency(0.1f)
-		, Lacunarity(2.0f)
 		, Persistence(0.5f)
+		, Lacunarity(2.0f)
 		, OctaveNumber(4)
-		, SurfaceResolution(10)
+		, PointSpacing(10)
 	{}
 
 	void TransferNoiseSettings(FNoiseSettings& NoiseSettingsOut);
 
-	/** Random number generator seed for repeatability */
+	/** Random number generator seed for repeatability. If the value is -1, a different random seed will be used every time, otherwise the specified seed will always be used */
 	UPROPERTY(EditAnywhere, Category = CommonFracture, meta = (DisplayName = "Random Seed", UIMin = "-1", UIMax = "1000", ClampMin = "-1"))
 	int32 RandomSeed;
 
-	/** Chance to shatter each selected bone.  */
+	/** Chance to fracture each selected bone. If 0, no bones will fracture; if 1, all bones will fracture. */
 	UPROPERTY(EditAnywhere, Category = CommonFracture, meta = (DisplayName = "Chance To Fracture Per Bone", UIMin = "0.0", UIMax = "1.0", ClampMin = "0.0", ClampMax = "1.0"))
 	float ChanceToFracture;
 
@@ -71,29 +71,32 @@ public:
 	UPROPERTY(EditAnywhere, Category = CommonFracture, meta = (UIMin = "0.0", ClampMin = "0.0"))
 	float Grout = 0.0f;
 
-	/** Size of the noise displacement in centimeters */
+	/** Size of the Perlin noise displacement (in cm). If 0, no noise will be applied */
 	UPROPERTY(EditAnywhere, Category = Noise, meta = (UIMin = "0.0"))
 	float Amplitude;
 
-	/** Period of the Perlin noise.  Smaller values will create noise faces that are smoother */
+	/** Period of the Perlin noise.  Smaller values will create a smoother noise pattern */
 	UPROPERTY(EditAnywhere, Category = Noise)
 	float Frequency;
 
-	/** Lacunarity of the Perlin noise.  Controls how the frequency scales per octave. */
-	UPROPERTY(EditAnywhere, Category = Noise, meta = (UIMin = "1.0", UIMax = "4.0"))
-	float Lacunarity;
-
-	/** Persistence of the Perlin noise.  Controls how the amplitude scales per octave. */
+	/** Persistence of the layers of Perlin noise. At each layer (octave) after the first, the amplitude of the Perlin noise is scaled by this factor */
 	UPROPERTY(EditAnywhere, Category = Noise, meta = (UIMin = "0.01", UIMax = "1.0"))
 	float Persistence;
 
-	/** Number of fractal layers of Perlin noise to apply.  Smaller values (1 or 2) will create noise that looks like gentle rolling hills, while larger values (> 4) will tend to look more like craggy mountains */
+	/** Lacunarity of the layers of Perlin noise. At each layer (octave) after the first, the frequency of the Perlin noise is scaled by this factor */
+	UPROPERTY(EditAnywhere, Category = Noise, meta = (UIMin = "1.0", UIMax = "4.0"))
+	float Lacunarity;
+
+	/** 
+	 * Number of fractal layers of Perlin noise to apply. Each layer is additive, with Amplitude and Frequency parameters scaled by Persistence and Lacunarity
+	 * Smaller values (1 or 2) will create noise that looks like gentle rolling hills, while larger values (> 4) will tend to look more like craggy mountains
+	 */
 	UPROPERTY(EditAnywhere, Category = Noise, meta = (ClampMin = "1", UIMax = "8"))
 	int32 OctaveNumber;
 
-	/** Spacing between vertices on cut surfaces, where noise is added.  Larger spacing between vertices will create more efficient meshes with fewer triangles, but less resolution to see the shape of the added noise  */
+	/** Distance (in cm) between vertices on cut surfaces where noise is added.  Larger spacing between vertices will create more efficient meshes with fewer triangles, but less resolution to see the shape of the added noise  */
 	UPROPERTY(EditAnywhere, Category = Noise, meta = (DisplayName = "Point Spacing", UIMin = "1", ClampMin = "0.1"))
-	float SurfaceResolution;
+	float PointSpacing;
 
 	/// Get the maximum distance a vertex could be moved by a combination of grout and noise
 	float GetMaxVertexMovement()
@@ -118,10 +121,28 @@ public:
 	UFractureCollisionSettings(const FObjectInitializer& ObjInit)
 	: Super(ObjInit) {}
 
-	/** Target spacing between collision samples on the mesh surface. */
-	UPROPERTY(EditAnywhere, Category = Collision, meta = (UIMin = "1", ClampMin = "0.1"))
+	/**
+	 * If enabled, add extra vertices (without triangles) to the geometry in regions where vertices are spaced too far apart (e.g. across large triangles)
+	 * These extra vertices will be used as collision samples in particle-implicit collisions, and can help the physics system detect collisions more accurately
+	 * 
+	 * Note this is *only* useful for simulations that use particle-implicit collisions
+	 */
+	UPROPERTY(EditAnywhere, Category = Collision)
+	bool bAddSamplesForCollision = false;
+
+	/**
+	 * The number of centimeters to allow between vertices on the mesh surface: If there are gaps larger than this, add additional vertices (without triangles) to help support particle-implicit collisions
+	 * Only used if Add Samples For Collision is enabled
+	 */
+	UPROPERTY(EditAnywhere, Category = Collision, meta = (UIMin = "1", ClampMin = "0.1", EditCondition = "bAddSamplesForCollision"))
 	float PointSpacing = 50.0f;
 
+	// Get point spacing or the special value of 0 if adding collision samples is disabled
+	float GetPointSpacing()
+	{
+		return bAddSamplesForCollision ? PointSpacing : 0.0f;
+	}
+
 	// TODO: add remeshing options here as well
 };
 
diff --git a/Engine/Plugins/Experimental/ChaosEditor/Source/FractureEditor/Private/FractureToolGenerators.cpp b/Engine/Plugins/Experimental/ChaosEditor/Source/FractureEditor/Private/FractureToolGenerators.cpp
--- a/Engine/Plugins/Experimental/ChaosEditor/Source/FractureEditor/Private/FractureToolGenerators.cpp
+++ b/Engine/Plugins/Experimental/ChaosEditor/Source/FractureEditor/Private/FractureToolGenerators.cpp
@@ -33,7 +33,7 @@ FText UFractureToolGenerateAsset::GetDisplayText() const
 
 FText UFractureToolGenerateAsset::GetTooltipText() const
 {
-	return FText(NSLOCTEXT("Fracture", "FractureToolGenerateAssetTooltip", "Generate Geometry Collection Asset from static meshes contained in selected actors."));
+	return FText(NSLOCTEXT("Fracture", "FractureToolGenerateAssetTooltip", "Generate a Geometry Collection Asset from selected Static Meshes and/or Geometry Collections."));
 }
 
 FSlateIcon UFractureToolGenerateAsset::GetToolIcon() const
@@ -43,7 +43,7 @@ FSlateIcon UFractureToolGenerateAsset::GetToolIcon() const
 
 void UFractureToolGenerateAsset::RegisterUICommand(FFractureEditorCommands* BindingContext)
 {
-	UI_COMMAND_EXT(BindingContext, UICommandInfo, "GenerateAsset", "New", "Generate a geometry collection asset from static meshes contained in selected actors.", EUserInterfaceActionType::ToggleButton, FInputChord());
+	UI_COMMAND_EXT(BindingContext, UICommandInfo, "GenerateAsset", "New", "Generate a new Geometry Collection Asset from the selected Static Meshes and/or Geometry Collections. Geometry Collections are assets that support fracture.", EUserInterfaceActionType::ToggleButton, FInputChord());
 	BindingContext->GenerateAsset = UICommandInfo;
 }
 
@@ -74,7 +74,7 @@ void UFractureToolGenerateAsset::OpenGenerateAssetDialog(TArray<AActor*>& Actors
 
 	SAssignNew(PickAssetPathWindow, SWindow)
 		.Title(LOCTEXT("SelectPath", "Select Path"))
-		.ToolTipText(LOCTEXT("SelectPathTooltip", "Select the path where the Geometry Collection will be created at"))
+		.ToolTipText(LOCTEXT("SelectPathTooltip", "Select the asset path for your new Geometry Collection"))
 		.ClientSize(FVector2D(500, 500));
 
 	// NOTE - the parent window has to completely exist before this one does so the parent gets set properly.
@@ -359,7 +359,7 @@ FText UFractureToolResetAsset::GetDisplayText() const
 
 FText UFractureToolResetAsset::GetTooltipText() const
 {
-	return FText(NSLOCTEXT("Fracture", "FractureToolResetTooltip", "Reset Geometry Collection Asset to its initial unfractured state."));
+	return FText(NSLOCTEXT("Fracture", "FractureToolResetTooltip", "Reset Geometry Collections to their initial unfractured states."));
 }
 
 FSlateIcon UFractureToolResetAsset::GetToolIcon() const
@@ -369,7 +369,7 @@ FSlateIcon UFractureToolResetAsset::GetToolIcon() const
 
 void UFractureToolResetAsset::RegisterUICommand(FFractureEditorCommands* BindingContext)
 {
-	UI_COMMAND_EXT(BindingContext, UICommandInfo, "ResetAsset", "Reset", "Reset selected geometry collection(s) to their initial unfractured state.", EUserInterfaceActionType::ToggleButton, FInputChord());
+	UI_COMMAND_EXT(BindingContext, UICommandInfo, "ResetAsset", "Reset", "Reset selected Geometry Collections to their initial unfractured states.", EUserInterfaceActionType::ToggleButton, FInputChord());
 	BindingContext->ResetAsset = UICommandInfo;
 }
 
diff --git a/Engine/Plugins/Experimental/ChaosEditor/Source/FractureEditor/Private/FractureToolMeshCut.cpp b/Engine/Plugins/Experimental/ChaosEditor/Source/FractureEditor/Private/FractureToolMeshCut.cpp
--- a/Engine/Plugins/Experimental/ChaosEditor/Source/FractureEditor/Private/FractureToolMeshCut.cpp
+++ b/Engine/Plugins/Experimental/ChaosEditor/Source/FractureEditor/Private/FractureToolMeshCut.cpp
@@ -45,9 +45,9 @@ void UFractureToolMeshCut::RegisterUICommand( FFractureEditorCommands* BindingCo
 TArray<UObject*> UFractureToolMeshCut::GetSettingsObjects() const
  {
 	TArray<UObject*> Settings;
+	Settings.Add(MeshCutSettings);
 	//Settings.Add(CutterSettings); // TODO: add cutter settings if/when we support noise and grout
 	Settings.Add(CollisionSettings);
-	Settings.Add(MeshCutSettings);
 	return Settings;
 }
 
@@ -213,7 +213,7 @@ int32 UFractureToolMeshCut::ExecuteFracture(const FFractureToolContext& Fracture
 		// (Note: noise and grout not currently supported)
 		if (LocalCutSettings->CutDistribution == EMeshCutDistribution::SingleCut)
 		{
-			return CutWithMesh(MeshDescription, ActorTransform, InternalSurfaceMaterials, *FractureContext.GetGeometryCollection(), FractureContext.GetSelection(), CollisionSettings->PointSpacing, FractureContext.GetTransform());
+			return CutWithMesh(MeshDescription, ActorTransform, InternalSurfaceMaterials, *FractureContext.GetGeometryCollection(), FractureContext.GetSelection(), CollisionSettings->GetPointSpacing(), FractureContext.GetTransform());
 		}
 		else
 		{
@@ -238,7 +238,7 @@ int32 UFractureToolMeshCut::ExecuteFracture(const FFractureToolContext& Fracture
 			for (const FTransform& ScatterTransform : MeshTransforms)
 			{
 				EnterProgressFrame(1);
-				int32 Index = CutWithMesh(MeshDescription, ScatterTransform, InternalSurfaceMaterials, *FractureContext.GetGeometryCollection(), BonesToCut, CollisionSettings->PointSpacing, FractureContext.GetTransform());
+				int32 Index = CutWithMesh(MeshDescription, ScatterTransform, InternalSurfaceMaterials, *FractureContext.GetGeometryCollection(), BonesToCut, CollisionSettings->GetPointSpacing(), FractureContext.GetTransform());
 				if (FirstIndex == -1)
 				{
 					FirstIndex = Index;
diff --git a/Engine/Plugins/Experimental/ChaosEditor/Source/FractureEditor/Private/FractureToolMeshCut.h b/Engine/Plugins/Experimental/ChaosEditor/Source/FractureEditor/Private/FractureToolMeshCut.h
--- a/Engine/Plugins/Experimental/ChaosEditor/Source/FractureEditor/Private/FractureToolMeshCut.h
+++ b/Engine/Plugins/Experimental/ChaosEditor/Source/FractureEditor/Private/FractureToolMeshCut.h
@@ -29,7 +29,7 @@ public:
 	UFractureMeshCutSettings(const FObjectInitializer& ObjInit)
 		: Super(ObjInit) {}
 
-	/** Static mesh actor to be used as a cutting surface; should be a closed, watertight mesh */
+	/** Static Mesh Actor to be used as a cutting surface. For best results, use a closed, watertight mesh */
 	UPROPERTY(EditAnywhere, Category = MeshCut, meta = (DisplayName = "Cutting Actor"))
 	TLazyObjectPtr<AStaticMeshActor> CuttingActor;
 
@@ -61,24 +61,27 @@ public:
 	UPROPERTY(EditAnywhere, Category = Distribution, meta = (DisplayName = "Variability", EditCondition = "CutDistribution == EMeshCutDistribution::Grid", EditConditionHides, UIMin = "0.0", ClampMin = "0.0"))
 	float Variability = 0.0f;
 
+	/** Minimum scale factor to apply to cutting meshes. A random scale will be chosen between Min and Max */
 	UPROPERTY(EditAnywhere, Category = Distribution, meta = (ClampMin = "0.001", EditCondition = "CutDistribution != EMeshCutDistribution::SingleCut", EditConditionHides))
 	float MinScaleFactor = .5;
 
+	/** Maximum scale factor to apply to cutting meshes. A random scale will be chosen between Min and Max */
 	UPROPERTY(EditAnywhere, Category = Distribution, meta = (ClampMin = "0.001", EditCondition = "CutDistribution != EMeshCutDistribution::SingleCut", EditConditionHides))
 	float MaxScaleFactor = 1.5;
 
+	/** Whether to randomly vary the orientation of the cutting meshes */
 	UPROPERTY(EditAnywhere, Category = Distribution, meta = (EditCondition = "CutDistribution != EMeshCutDistribution::SingleCut", EditConditionHides))
 	bool bRandomOrientation = true;
 
-	/** Roll will be chosen between -Range to +Range */
+	/** Roll will be chosen between -Range and +Range */
 	UPROPERTY(EditAnywhere, Category = Distribution, meta = (DisplayName = "+/- Roll Range", EditCondition = "CutDistribution != EMeshCutDistribution::SingleCut && bRandomOrientation", EditConditionHides, ClampMin = "0", ClampMax = "180"))
 	float RollRange = 180;
 
-	/** Pitch will be chosen between -Range to +Range */
+	/** Pitch will be chosen between -Range and +Range */
 	UPROPERTY(EditAnywhere, Category = Distribution, meta = (DisplayName = "+/- Pitch Range", EditCondition = "CutDistribution != EMeshCutDistribution::SingleCut && bRandomOrientation", EditConditionHides, ClampMin = "0", ClampMax = "180"))
 	float PitchRange = 180;
 
-	/** Yaw will be chosen between -Range to +Range */
+	/** Yaw will be chosen between -Range and +Range */
 	UPROPERTY(EditAnywhere, Category = Distribution, meta = (DisplayName = "+/- Yaw Range", EditCondition = "CutDistribution != EMeshCutDistribution::SingleCut && bRandomOrientation", EditConditionHides, ClampMin = "0", ClampMax = "180"))
 	float YawRange = 180;
 };
diff --git a/Engine/Plugins/Experimental/ChaosEditor/Source/FractureEditor/Private/FractureToolPlaneCut.cpp b/Engine/Plugins/Experimental/ChaosEditor/Source/FractureEditor/Private/FractureToolPlaneCut.cpp
--- a/Engine/Plugins/Experimental/ChaosEditor/Source/FractureEditor/Private/FractureToolPlaneCut.cpp
+++ b/Engine/Plugins/Experimental/ChaosEditor/Source/FractureEditor/Private/FractureToolPlaneCut.cpp
@@ -116,10 +116,10 @@ void UFractureToolPlaneCut::Render(const FSceneView* View, FViewport* Viewport,
 TArray<UObject*> UFractureToolPlaneCut::GetSettingsObjects() const
  {
 	TArray<UObject*> Settings;
-	Settings.Add(CutterSettings);
+	Settings.Add(PlaneCutSettings);
 	Settings.Add(GizmoSettings);
+	Settings.Add(CutterSettings);
 	Settings.Add(CollisionSettings);
-	Settings.Add(PlaneCutSettings);
 	return Settings;
 }
 
@@ -197,7 +197,7 @@ int32 UFractureToolPlaneCut::ExecuteFracture(const FFractureToolContext& Fractur
 		// Proximity is invalidated.
 		ClearProximity(FractureContext.GetGeometryCollection().Get());
 
-		return CutMultipleWithMultiplePlanes(CuttingPlanes, InternalSurfaceMaterials, *FractureContext.GetGeometryCollection(), FractureContext.GetSelection(), CutterSettings->Grout, CollisionSettings->PointSpacing, FractureContext.GetSeed(), FractureContext.GetTransform());
+		return CutMultipleWithMultiplePlanes(CuttingPlanes, InternalSurfaceMaterials, *FractureContext.GetGeometryCollection(), FractureContext.GetSelection(), CutterSettings->Grout, CollisionSettings->GetPointSpacing(), FractureContext.GetSeed(), FractureContext.GetTransform());
 	}
 
 	return INDEX_NONE;
diff --git a/Engine/Plugins/Experimental/ChaosEditor/Source/FractureEditor/Private/FractureToolPlaneCut.h b/Engine/Plugins/Experimental/ChaosEditor/Source/FractureEditor/Private/FractureToolPlaneCut.h
--- a/Engine/Plugins/Experimental/ChaosEditor/Source/FractureEditor/Private/FractureToolPlaneCut.h
+++ b/Engine/Plugins/Experimental/ChaosEditor/Source/FractureEditor/Private/FractureToolPlaneCut.h
@@ -19,7 +19,7 @@ public:
 		: Super(ObjInit)
 		, NumberPlanarCuts(1) {}
 
-	/** Number of Cutting Planes (Note: Only applicable if interactive placement gizmo is not in use) */
+	/** Number of cutting planes. Only used when "Use Gizmo" is disabled */
 	UPROPERTY(EditAnywhere, Category = PlaneCut, meta = (DisplayName = "Number of Cuts", UIMin = "1", UIMax = "20", ClampMin = "1", EditCondition = "bCanCutWithMultiplePlanes", HideEditConditionToggle))
 	int32 NumberPlanarCuts;
 
diff --git a/Engine/Plugins/Experimental/ChaosEditor/Source/FractureEditor/Private/FractureToolProperties.h b/Engine/Plugins/Experimental/ChaosEditor/Source/FractureEditor/Private/FractureToolProperties.h
--- a/Engine/Plugins/Experimental/ChaosEditor/Source/FractureEditor/Private/FractureToolProperties.h
+++ b/Engine/Plugins/Experimental/ChaosEditor/Source/FractureEditor/Private/FractureToolProperties.h
@@ -10,9 +10,9 @@ UENUM(BlueprintType)
 enum class EDynamicStateOverrideEnum : uint8
 {
 	NoOverride = 0										  UMETA(DisplayName = "No Override"),
-	Sleeping = 1 /*Chaos::EObjectStateType::Sleeping*/    UMETA(DisplayName = "Sleeping"),
-	Kinematic = 2 /*Chaos::EObjectStateType::Kinematic*/  UMETA(DisplayName = "Kinematic"),
-	Static = 3    /*Chaos::EObjectStateType::Static*/     UMETA(DisplayName = "Static")
+	Sleeping = 1 /*~Chaos::EObjectStateType::Sleeping*/    UMETA(DisplayName = "Sleeping"),
+	Kinematic = 2 /*~Chaos::EObjectStateType::Kinematic*/  UMETA(DisplayName = "Kinematic"),
+	Static = 3    /*~Chaos::EObjectStateType::Static*/     UMETA(DisplayName = "Static")
 };
 
 /** Settings specifically related to the one-time destructive fracturing of a mesh **/
@@ -28,7 +28,7 @@ public:
 		, InitialDynamicState(EDynamicStateOverrideEnum::Kinematic)
 	{}
 
-	/** Random number generator seed for repeatability */
+	/** Simulation state to be set on selected bones */
 	UPROPERTY(EditAnywhere, Category = SetInitialDynamicState, meta = (DisplayName = "Initial Dynamic State"))
 	EDynamicStateOverrideEnum InitialDynamicState;
 
diff --git a/Engine/Plugins/Experimental/ChaosEditor/Source/FractureEditor/Private/FractureToolRadial.cpp b/Engine/Plugins/Experimental/ChaosEditor/Source/FractureEditor/Private/FractureToolRadial.cpp
--- a/Engine/Plugins/Experimental/ChaosEditor/Source/FractureEditor/Private/FractureToolRadial.cpp
+++ b/Engine/Plugins/Experimental/ChaosEditor/Source/FractureEditor/Private/FractureToolRadial.cpp
@@ -53,11 +53,11 @@ void UFractureToolRadial::RegisterUICommand( FFractureEditorCommands* BindingCon
 
 TArray<UObject*> UFractureToolRadial::GetSettingsObjects() const 
 { 
-	TArray<UObject*> Settings; 
-	Settings.Add(CutterSettings);
+	TArray<UObject*> Settings;
+	Settings.Add(RadialSettings);
 	Settings.Add(GizmoSettings);
+	Settings.Add(CutterSettings);
 	Settings.Add(CollisionSettings);
-	Settings.Add(RadialSettings);
 	return Settings;
 }
 
@@ -87,7 +87,7 @@ void UFractureToolRadial::GenerateVoronoiSites(const FFractureToolContext& Conte
 	FVector::FReal Len = RadialStep * .5;
 	for (int32 ii = 0; ii < RadialSettings->RadialSteps; ++ii, Len += RadialStep)
 	{
-		FVector::FReal Angle = RadialSettings->AngleOffset;
+		FVector::FReal Angle = FMath::DegreesToRadians(RadialSettings->AngleOffset);
 		for (int32 kk = 0; kk < RadialSettings->AngularSteps; ++kk, Angle += AngularStep)
 		{
 			FVector RotatingOffset = Len * (FMath::Cos(Angle) * BasisX + FMath::Sin(Angle) * BasisY);
diff --git a/Engine/Plugins/Experimental/ChaosEditor/Source/FractureEditor/Private/FractureToolRadial.h b/Engine/Plugins/Experimental/ChaosEditor/Source/FractureEditor/Private/FractureToolRadial.h
--- a/Engine/Plugins/Experimental/ChaosEditor/Source/FractureEditor/Private/FractureToolRadial.h
+++ b/Engine/Plugins/Experimental/ChaosEditor/Source/FractureEditor/Private/FractureToolRadial.h
@@ -26,18 +26,18 @@ public:
 		, Variability(0.0f)
 	{}
 
-	/** Center of generated pattern.  Only used when "Use Gizmo" is disabled. */
+	/** Center of generated pattern. Only used when "Use Gizmo" is disabled */
 	UPROPERTY(EditAnywhere, Category = RadialVoronoi, meta = (EditCondition = "!bPositionedByGizmo", HideEditConditionToggle))
 	FVector Center;
 
-	/** Normal to plane in which sites are generated.  Only used when "Use Gizmo" is disabled. */
+	/** Normal to plane in which sites are generated. Only used when "Use Gizmo" is disabled */
 	UPROPERTY(EditAnywhere, Category = RadialVoronoi, meta = (EditCondition = "!bPositionedByGizmo", HideEditConditionToggle))
 	FVector Normal;
 
 	UPROPERTY()
 	bool bPositionedByGizmo = true;
 
-	/** Pattern radius */
+	/** Pattern radius (in cm) */
 	UPROPERTY(EditAnywhere, Category = RadialVoronoi, meta = (DisplayName = "Radius", UIMin = "0.0", ClampMin = "0.0"))
 	float Radius;
 
@@ -49,11 +49,11 @@ public:
 	UPROPERTY(EditAnywhere, Category = RadialVoronoi, meta = (DisplayName = "Radial Steps", UIMin = "1", UIMax = "50", ClampMin = "1"))
 	int RadialSteps;
 
-	/** Angle offset at each radial step */
-	UPROPERTY(EditAnywhere, Category = RadialVoronoi, meta = (DisplayName = "Angle Offset", UIMin = "0.0", ClampMin = "0.0"))
+	/** Angle offset at each radial step (in degrees) */
+	UPROPERTY(EditAnywhere, Category = RadialVoronoi, meta = (DisplayName = "Angle Offset"))
 	float AngleOffset;
 
-	/** Randomness of sites distribution */
+	/** Amount to randomly displace each Voronoi site (in cm) */
 	UPROPERTY(EditAnywhere, Category = RadialVoronoi, meta = (DisplayName = "Variability", UIMin = "0.0", ClampMin = "0.0"))
 	float Variability;
 };
diff --git a/Engine/Plugins/Experimental/ChaosEditor/Source/FractureEditor/Private/FractureToolResample.cpp b/Engine/Plugins/Experimental/ChaosEditor/Source/FractureEditor/Private/FractureToolResample.cpp
--- a/Engine/Plugins/Experimental/ChaosEditor/Source/FractureEditor/Private/FractureToolResample.cpp
+++ b/Engine/Plugins/Experimental/ChaosEditor/Source/FractureEditor/Private/FractureToolResample.cpp
@@ -36,7 +36,7 @@ FSlateIcon UFractureToolResample::GetToolIcon() const
 
 void UFractureToolResample::RegisterUICommand( FFractureEditorCommands* BindingContext ) 
 {
-	UI_COMMAND_EXT( BindingContext, UICommandInfo, "Resample", "Resmpl", "Resample to add collision particles in large flat regions that otherwise might have poor collision response.", EUserInterfaceActionType::ToggleButton, FInputChord() );
+	UI_COMMAND_EXT( BindingContext, UICommandInfo, "Resample", "Resmpl", "Resample to add collision particles in large flat regions that otherwise might have poor collision response. Only useful to help improve Particle-Implicit collisions.", EUserInterfaceActionType::ToggleButton, FInputChord() );
 	BindingContext->Resample = UICommandInfo;
 }
 
@@ -104,9 +104,7 @@ int32 UFractureToolResample::ExecuteFracture(const FFractureToolContext& Fractur
 {
 	if (FractureContext.IsValid())
 	{
-		const UFractureCollisionSettings* LocalCutSettings = CollisionSettings;
-
-		return AddCollisionSampleVertices(LocalCutSettings->PointSpacing, *FractureContext.GetGeometryCollection(), FractureContext.GetSelection());
+		return AddCollisionSampleVertices(CollisionSettings->GetPointSpacing(), *FractureContext.GetGeometryCollection(), FractureContext.GetSelection());
 	}
 
 	return INDEX_NONE;
diff --git a/Engine/Plugins/Experimental/ChaosEditor/Source/FractureEditor/Private/FractureToolSelectors.cpp b/Engine/Plugins/Experimental/ChaosEditor/Source/FractureEditor/Private/FractureToolSelectors.cpp
--- a/Engine/Plugins/Experimental/ChaosEditor/Source/FractureEditor/Private/FractureToolSelectors.cpp
+++ b/Engine/Plugins/Experimental/ChaosEditor/Source/FractureEditor/Private/FractureToolSelectors.cpp
@@ -20,7 +20,7 @@ FText UFractureToolSelectAll::GetDisplayText() const
 
 FText UFractureToolSelectAll::GetTooltipText() const
 {
-	return FText(NSLOCTEXT("Fracture", "FractureToolSelectAllTooltip", "Selects all Bones in the GeometryCollection"));
+	return FText(NSLOCTEXT("Fracture", "FractureToolSelectAllTooltip", "Selects all bones in the Geometry Collection"));
 }
 
 FSlateIcon UFractureToolSelectAll::GetToolIcon() const
@@ -30,7 +30,7 @@ FSlateIcon UFractureToolSelectAll::GetToolIcon() const
 
 void UFractureToolSelectAll::RegisterUICommand(FFractureEditorCommands* BindingContext)
 {
-	UI_COMMAND_EXT(BindingContext, UICommandInfo, "SelectAll", "All", "Selects all Bones in the GeometryCollection.", EUserInterfaceActionType::Button, FInputChord(EModifierKey::Control | EModifierKey::Shift, EKeys::A));
+	UI_COMMAND_EXT(BindingContext, UICommandInfo, "SelectAll", "All", "Selects all bones in the Geometry Collection", EUserInterfaceActionType::Button, FInputChord(EModifierKey::Control | EModifierKey::Shift, EKeys::A));
 	BindingContext->SelectAll = UICommandInfo;
 }
 
@@ -97,7 +97,7 @@ FText UFractureToolSelectNone::GetDisplayText() const
 
 FText UFractureToolSelectNone::GetTooltipText() const
 {
-	return FText(NSLOCTEXT("Fracture", "FractureToolSelectNoneTooltip", "Deselects all Bones in the GeometryCollection."));
+	return FText(NSLOCTEXT("Fracture", "FractureToolSelectNoneTooltip", "Deselects all bones in the Geometry Collection"));
 }
 
 FSlateIcon UFractureToolSelectNone::GetToolIcon() const
@@ -107,7 +107,7 @@ FSlateIcon UFractureToolSelectNone::GetToolIcon() const
 
 void UFractureToolSelectNone::RegisterUICommand(FFractureEditorCommands* BindingContext)
 {
-	UI_COMMAND_EXT(BindingContext, UICommandInfo, "SelectNone", "None", "Deselects all Bones in the GeometryCollection.", EUserInterfaceActionType::Button, FInputChord(EModifierKey::Control | EModifierKey::Shift, EKeys::D));
+	UI_COMMAND_EXT(BindingContext, UICommandInfo, "SelectNone", "None", "Deselects all bones in the Geometry Collection", EUserInterfaceActionType::Button, FInputChord(EModifierKey::Control | EModifierKey::Shift, EKeys::D));
 	BindingContext->SelectNone = UICommandInfo;
 }
 
diff --git a/Engine/Plugins/Experimental/ChaosEditor/Source/FractureEditor/Private/FractureToolSlice.cpp b/Engine/Plugins/Experimental/ChaosEditor/Source/FractureEditor/Private/FractureToolSlice.cpp
--- a/Engine/Plugins/Experimental/ChaosEditor/Source/FractureEditor/Private/FractureToolSlice.cpp
+++ b/Engine/Plugins/Experimental/ChaosEditor/Source/FractureEditor/Private/FractureToolSlice.cpp
@@ -49,9 +49,9 @@ void UFractureToolSlice::RegisterUICommand( FFractureEditorCommands* BindingCont
 TArray<UObject*> UFractureToolSlice::GetSettingsObjects() const 
 { 
 	TArray<UObject*> Settings; 
+	Settings.Add(SliceSettings);
 	Settings.Add(CutterSettings);
 	Settings.Add(CollisionSettings);
-	Settings.Add(SliceSettings);
 	return Settings;
 }
 
@@ -185,7 +185,7 @@ int32 UFractureToolSlice::ExecuteFracture(const FFractureToolContext& FractureCo
 		// Proximity is invalidated.
 		ClearProximity(FractureContext.GetGeometryCollection().Get());
 
-		return CutMultipleWithMultiplePlanes(CuttingPlanes, InternalSurfaceMaterials, *FractureContext.GetGeometryCollection(), FractureContext.GetSelection(), CutterSettings->Grout, CollisionSettings->PointSpacing, FractureContext.GetSeed(), FractureContext.GetTransform());
+		return CutMultipleWithMultiplePlanes(CuttingPlanes, InternalSurfaceMaterials, *FractureContext.GetGeometryCollection(), FractureContext.GetSelection(), CutterSettings->Grout, CollisionSettings->GetPointSpacing(), FractureContext.GetSeed(), FractureContext.GetTransform());
 	}
 
 	return INDEX_NONE;
diff --git a/Engine/Plugins/Experimental/ChaosEditor/Source/FractureEditor/Private/FractureToolSlice.h b/Engine/Plugins/Experimental/ChaosEditor/Source/FractureEditor/Private/FractureToolSlice.h
--- a/Engine/Plugins/Experimental/ChaosEditor/Source/FractureEditor/Private/FractureToolSlice.h
+++ b/Engine/Plugins/Experimental/ChaosEditor/Source/FractureEditor/Private/FractureToolSlice.h
@@ -25,23 +25,23 @@ public:
 		, SliceOffsetVariation(0.0f)
 	{}
 
-	/** Num Slices X axis - Slicing Method */
+	/** Number of slices along the X axis */
 	UPROPERTY(EditAnywhere, Category = Slicing, meta = (UIMin = "0"))
 	int32 SlicesX;
 
-	/** Num Slices Y axis - Slicing Method */
+	/** Number of slices along the Y axis */
 	UPROPERTY(EditAnywhere, Category = Slicing, meta = (UIMin = "0"))
 	int32 SlicesY;
 
-	/** Num Slices Z axis - Slicing Method */
+	/** Number of slices along the Z axis */
 	UPROPERTY(EditAnywhere, Category = Slicing, meta = (UIMin = "0"))
 	int32 SlicesZ;
 
-	/** Slicing Angle Variation - Slicing Method [0..1] */
+	/** Maximum angle (in degrees) to randomly rotate each slicing plane */
 	UPROPERTY(EditAnywhere, Category = Slicing, meta = (DisplayName = "Random Angle Variation", UIMin = "0.0", UIMax = "90.0"))
 	float SliceAngleVariation;
 
-	/** Slicing Offset Variation - Slicing Method [0..1] */
+	/** Maximum distance (in cm) to randomly shift each slicing plane */
 	UPROPERTY(EditAnywhere, Category = Slicing, meta = (DisplayName = "Random Offset Variation", UIMin = "0.0"))
 	float SliceOffsetVariation;
 };
diff --git a/Engine/Plugins/Experimental/ChaosEditor/Source/FractureEditor/Private/FractureToolUV.cpp b/Engine/Plugins/Experimental/ChaosEditor/Source/FractureEditor/Private/FractureToolUV.cpp
--- a/Engine/Plugins/Experimental/ChaosEditor/Source/FractureEditor/Private/FractureToolUV.cpp
+++ b/Engine/Plugins/Experimental/ChaosEditor/Source/FractureEditor/Private/FractureToolUV.cpp
@@ -60,6 +60,12 @@ void UFractureAutoUVSettings::ChangeNumUVChannels(int32 Delta)
 	}
 }
 
+void UFractureAutoUVSettings::DisableBoneColors()
+{
+	UFractureToolAutoUV* AutoUVTool = Cast<UFractureToolAutoUV>(OwnerTool.Get());
+	AutoUVTool->DisableBoneColors();
+}
+
 void UFractureAutoUVSettings::BoxProjectUVs()
 {
 	UFractureToolAutoUV* AutoUVTool = Cast<UFractureToolAutoUV>(OwnerTool.Get());
@@ -152,8 +158,7 @@ void UFractureToolAutoUV::UpdateUVChannels(int32 TargetNumUVChannels)
 			// but if we're increasing the number of channels, some geometry collections may have more channels already and we leave those alone
 			if (NumChannels < TargetNumUVChannels || !bIsIncreasing)
 			{
-				FGeometryCollectionEdit Edit(GeometryCollectionComponent, GeometryCollection::EEditUpdate::None);
-				Edit.GetRestCollection()->Modify();
+				FGeometryCollectionEdit Edit(GeometryCollectionComponent, GeometryCollection::EEditUpdate::Rest, true);
 				GeometryCollectionComponent->GetRestCollection()->GetGeometryCollection()->SetNumUVLayers(TargetNumUVChannels);
 				if (bIsIncreasing)
 				{
@@ -177,6 +182,20 @@ void UFractureToolAutoUV::UpdateUVChannels(int32 TargetNumUVChannels)
 	AutoUVSettings->SetNumUVChannels(MinUVChannels);
 }
 
+void UFractureToolAutoUV::DisableBoneColors()
+{
+	TSet<UGeometryCollectionComponent*> GeomCompSelection;
+	GetSelectedGeometryCollectionComponents(GeomCompSelection);
+
+	FScopedTransaction Transaction(LOCTEXT("DisableBoneColors", "Disable Bone Colors"), !GeomCompSelection.IsEmpty());
+	for (UGeometryCollectionComponent* GeometryCollectionComponent : GeomCompSelection)
+	{
+		GeometryCollectionComponent->Modify();
+		GeometryCollectionComponent->SetShowBoneColors(false);
+		GeometryCollectionComponent->MarkRenderStateDirty();
+	}
+}
+
 void UFractureToolAutoUV::BoxProjectUVs()
 {
 	int32 UVLayer = AutoUVSettings->GetSelectedChannelIndex();
@@ -192,8 +211,7 @@ void UFractureToolAutoUV::BoxProjectUVs()
 	FScopedTransaction Transaction(LOCTEXT("BoxProjectUVs", "Box Project UVs"), !GeomCompSelection.IsEmpty());
 	for (UGeometryCollectionComponent* GeometryCollectionComponent : GeomCompSelection)
 	{
-		FGeometryCollectionEdit Edit(GeometryCollectionComponent, GeometryCollection::EEditUpdate::None);
-		Edit.GetRestCollection()->Modify();
+		FGeometryCollectionEdit Edit(GeometryCollectionComponent, GeometryCollection::EEditUpdate::Rest, true);
 		UE::PlanarCut::BoxProjectUVs(UVLayer,
 			*GeometryCollectionComponent->GetRestCollection()->GetGeometryCollection(),
 			(FVector3d)AutoUVSettings->ProjectionScale,
diff --git a/Engine/Plugins/Experimental/ChaosEditor/Source/FractureEditor/Private/FractureToolUV.h b/Engine/Plugins/Experimental/ChaosEditor/Source/FractureEditor/Private/FractureToolUV.h
--- a/Engine/Plugins/Experimental/ChaosEditor/Source/FractureEditor/Private/FractureToolUV.h
+++ b/Engine/Plugins/Experimental/ChaosEditor/Source/FractureEditor/Private/FractureToolUV.h
@@ -40,13 +40,17 @@ enum class ETextureType
 UENUM()
 enum class ETargetMaterialIDs
 {
+	// Target faces with odd-numbered material IDs (corresponding to internal fracture faces)
 	OddIDs,
+	// Target faces with material IDs matching a custom list, as well as those with odd-numbered IDs (corresponding to internal fracture faces)
 	OddAndSelectedIDs,
+	// Target faces with material IDs matching a custom list
 	SelectedIDs,
+	// Target all faces
 	AllIDs
 };
 
-/** Settings specifically related to the one-time destructive fracturing of a mesh **/
+/** Settings for UV layout and texture baking on the geometry collection **/
 UCLASS(config = EditorPerProjectUserSettings)
 class UFractureAutoUVSettings : public UFractureToolSettings
 {
@@ -60,6 +64,7 @@ public:
 		SetNumUVChannels(1);
 	}
 
+	/** Which UV channel to use for layout and baking */
 	UPROPERTY(EditAnywhere, Category = UVChannel, meta = (DisplayName = "UV Channel", GetOptions = GetUVChannelNamesFunc))
 	FString UVChannel;
 
@@ -75,12 +80,14 @@ public:
 	void SetNumUVChannels(int32 NumUVChannels);
 	int32 GetSelectedChannelIndex(bool bForceToZeroOnFailure = true);
 
+	/** Add a UV channel to the Geometry Collection */
 	UFUNCTION(CallInEditor, Category = UVChannel, meta = (DisplayName = "Add UV Channel"))
 	void AddUVChannel()
 	{
 		ChangeNumUVChannels(1);
 	}
 
+	/** Remove the last UV channel from the Geometry Collection */
 	UFUNCTION(CallInEditor, Category = UVChannel, meta = (DisplayName = "Delete UV Channel"))
 	void DeleteUVChannel()
 	{
@@ -89,18 +96,34 @@ public:
 
 	void ChangeNumUVChannels(int32 Delta);
 
+	/** Immediately layout and update the UVs of the selected Geometry Collection, for faces with the Target Material IDs */
 	UFUNCTION(CallInEditor, Category = Operations, meta = (DisplayName = "Layout UVs", DisplayPriority = 1))
 	void LayoutUVs();
 
+	/** Bake a new texture with the chosen attributes, aligned to the current UV layout */
 	UFUNCTION(CallInEditor, Category = Operations, meta = (DisplayName = "Bake Texture", DisplayPriority = 2))
 	void BakeTexture();
 
+	/** Disable "Show Bone Colors" on the selected Geometry Collection Components, so textures are visible. This setting can also be toggled directly on the Component itself */
+	UFUNCTION(CallInEditor, Category = Operations)
+	void DisableBoneColors();
+
+	/** The scale factor to use for UV box projection */
 	UPROPERTY(EditAnywhere, Category = Unwrap)
 	FVector ProjectionScale = FVector(100, 100, 100);
 
+	/** Immediately perform a box projection to update the internal UVs of the selected Geometry Collection */
 	UFUNCTION(CallInEditor, Category = Unwrap, meta = (DisplayName = "Box Project UVs"))
 	void BoxProjectUVs();
 
+	/** Choose whether to texture only faces with odd material IDs (corresponding to internal faces) or a custom selection */
+	UPROPERTY(EditAnywhere, Category = MapSettings)
+	ETargetMaterialIDs TargetMaterialIDs;
+
+	/** Custom selection of material IDs to target for texturing */
+	UPROPERTY(EditAnywhere, Category = MapSettings, meta = (EditCondition = "TargetMaterialIDs != ETargetMaterialIDs::OddIDs && TargetMaterialIDs != ETargetMaterialIDs::AllIDs", EditConditionHides))
+	TArray<int32> MaterialIDs;
+
 	/** The pixel resolution of the generated map */
 	UPROPERTY(EditAnywhere, Category = MapSettings)
 	EAutoUVTextureResolution Resolution = EAutoUVTextureResolution::Resolution512;
@@ -113,15 +136,7 @@ public:
 	UPROPERTY(VisibleAnywhere, Category = MapSettings)
 	TObjectPtr<UTexture2D> Result;
 
-	/** Choose whether to texture only odd material IDs (corresponding to internal faces) or a custom selection */
-	UPROPERTY(EditAnywhere, Category = MapSettings)
-	ETargetMaterialIDs TargetMaterialIDs;
-
-	/** Custom selection of material IDs to target for texturing */
-	UPROPERTY(EditAnywhere, Category = MapSettings, meta = (EditCondition = "TargetMaterialIDs != ETargetMaterialIDs::OddIDs", EditConditionHides))
-	TArray<int32> MaterialIDs;
-
-	/** Whether to prompt user for an asset name for each generated texture, or automatically place them next to the source geometry collections */
+	/** Whether to prompt user for an asset name for each generated texture, or automatically place them next to the source Geometry Collections */
 	UPROPERTY(EditAnywhere, Category = MapSettings)
 	bool bPromptToSave = true;
 
@@ -143,10 +158,11 @@ public:
 
 	/**
 	 * Bake a smoothed curvature metric to a texture channel (blue)
-	 * Specifically, this is the mean curvature of a smoothed copy of each fractured piece, baked back to the respective fracture piece.
+	 * Specifically, this is the mean curvature of a smoothed copy of each fractured piece, baked back to the respective fracture piece
+	 * Note that this attribute is relatively expensive to compute
 	 */
 	UPROPERTY(EditAnywhere, Category = AttributesToBake, meta = (EditCondition = "BakeTextureType == ETextureType::ThicknessAndSurfaceAttributes", EditConditionHides))
-	bool bSmoothedCurvature = true;
+	bool bSmoothedCurvature = false;
 
 	/** Max distance to search for the outer mesh surface */
 	UPROPERTY(EditAnywhere, Category = DistToOuterSettings, meta = (EditCondition = "BakeTextureType == ETextureType::ThicknessAndSurfaceAttributes && bDistToOuter", EditConditionHides, UIMin = "1", UIMax = "100", ClampMin = ".01", ClampMax = "1000"))
@@ -176,7 +192,7 @@ public:
 	UPROPERTY(EditAnywhere, Category = SmoothedCurvatureSettings, meta = (EditCondition = "BakeTextureType == ETextureType::ThicknessAndSurfaceAttributes && bSmoothedCurvature", EditConditionHides, UIMin = "2", UIMax = "10.0", ClampMin = "1", ClampMax = "100.0"))
 	double ThicknessFactor = 4;
 
-	/** Curvatures in the range [-MaxCurvature, MaxCurvature] will be mapped from [0,1]. Values outside that range will be clamped. */
+	/** Curvatures in the range [-MaxCurvature, MaxCurvature] will be mapped from [0,1]. Values outside that range will be clamped */
 	UPROPERTY(EditAnywhere, Category = SmoothedCurvatureSettings, meta = (EditCondition = "BakeTextureType == ETextureType::ThicknessAndSurfaceAttributes && bSmoothedCurvature", EditConditionHides, UIMin = ".01", UIMax = "1", ClampMin = ".0001", ClampMax = "10"))
 	double MaxCurvature = .1;
 
@@ -199,7 +215,7 @@ public:
 	virtual FText GetDisplayText() const override;
 	virtual FText GetTooltipText() const override;
 	/** This is the Text that will appear on the button to execute the fracture **/
-	virtual FText GetApplyText() const override { return FText(NSLOCTEXT("AutoUV", "ExecuteAutoUV", "AutoUV")); }
+	virtual FText GetApplyText() const override { return FText(NSLOCTEXT("AutoUV", "ExecuteAutoUV", "Layout & Bake")); }
 	virtual FSlateIcon GetToolIcon() const override;
 	virtual void RegisterUICommand(FFractureEditorCommands* BindingContext) override;
 	virtual TArray<UObject*> GetSettingsObjects() const override;
@@ -224,6 +240,8 @@ public:
 
 	// Create new UV islands via box projection for the selected geometry collections
 	void BoxProjectUVs();
+
+	void DisableBoneColors();
 	
 	void LayoutUVs();
 	void BakeTexture();
diff --git a/Engine/Plugins/Experimental/ChaosEditor/Source/FractureEditor/Private/FractureToolUniform.h b/Engine/Plugins/Experimental/ChaosEditor/Source/FractureEditor/Private/FractureToolUniform.h
--- a/Engine/Plugins/Experimental/ChaosEditor/Source/FractureEditor/Private/FractureToolUniform.h
+++ b/Engine/Plugins/Experimental/ChaosEditor/Source/FractureEditor/Private/FractureToolUniform.h
@@ -20,12 +20,12 @@ public:
 		, NumberVoronoiSitesMax(20)
 	{}
 
-	/** Minimum Number of Voronoi sites - A random number will be chosen between the Min and Max for each bone you have selected */
-	UPROPERTY(EditAnywhere, Category = UniformVoronoi, meta = (DisplayName = "Minimum Voronoi Sites", UIMin = "1", UIMax = "5000", ClampMin = "1"))
+	/** Minimum Number of Voronoi sites. The amount of sites per Voronoi diagram will be chosen at random between Min and Max */
+	UPROPERTY(EditAnywhere, Category = UniformVoronoi, meta = (DisplayName = "Min Voronoi Sites", UIMin = "1", UIMax = "5000", ClampMin = "1"))
 	int32 NumberVoronoiSitesMin;
 
-	/** Maximum Number of Voronoi sites - A random number will be chosen between the Min and Max for each bone you have selected */
-	UPROPERTY(EditAnywhere, Category = UniformVoronoi, meta = (DisplayName = "Maximum Voronoi Sites", UIMin = "1", UIMax = "5000", ClampMin = "1"))
+	/** Maximum Number of Voronoi sites. The amount of sites per Voronoi diagram will be chosen at random between Min and Max */
+	UPROPERTY(EditAnywhere, Category = UniformVoronoi, meta = (DisplayName = "Max Voronoi Sites", UIMin = "1", UIMax = "5000", ClampMin = "1"))
 	int32 NumberVoronoiSitesMax;
 
 };
diff --git a/Engine/Plugins/Experimental/ChaosEditor/Source/FractureEditor/Public/FractureSettings.h b/Engine/Plugins/Experimental/ChaosEditor/Source/FractureEditor/Public/FractureSettings.h
--- a/Engine/Plugins/Experimental/ChaosEditor/Source/FractureEditor/Public/FractureSettings.h
+++ b/Engine/Plugins/Experimental/ChaosEditor/Source/FractureEditor/Public/FractureSettings.h
@@ -15,15 +15,15 @@ class UFractureSettings: public UObject
 public:
 	UFractureSettings(const FObjectInitializer& ObjInit);
 
-	/** How much to seperate the drawing of the bones  */
+	/** Amount to expand the displayed Geometry Collection bones into an 'exploded view' */
 	UPROPERTY(EditAnywhere, Category = ViewSettings, meta = (DisplayName = "Explode Amount", UIMin = "0.0", UIMax = "1.0", ClampMin = "0.0", ClampMax = "1.0"))
 	float ExplodeAmount;
 
-	/** Current level of the geometry collection displayed  */
+	/** Current level of the Geometry Collection displayed */
 	UPROPERTY(EditAnywhere, Category = ViewSettings, meta = (DisplayName = "Fracture Level", UIMin = "-1" ))
 	int32 FractureLevel;
 
-	/** When active, only show selected bones.  */
+	/** When active, only show selected bones */
 	UPROPERTY(EditAnywhere, Category = ViewSettings, meta = (DisplayName = "Hide Unselected"))
 	bool bHideUnselected;
 
