Use the correct UEdMode GetInteractiveToolsContext() for fracture mode tools, rather than sometimes using the global one
 + also fix snapping for the fracture mode gizmo

diff --git a/Engine/Plugins/Experimental/ChaosEditor/Source/FractureEditor/Private/FractureEditorMode.cpp b/Engine/Plugins/Experimental/ChaosEditor/Source/FractureEditor/Private/FractureEditorMode.cpp
--- a/Engine/Plugins/Experimental/ChaosEditor/Source/FractureEditor/Private/FractureEditorMode.cpp
+++ b/Engine/Plugins/Experimental/ChaosEditor/Source/FractureEditor/Private/FractureEditorMode.cpp
@@ -26,6 +26,7 @@
 #include "GeometryCollection/GeometryCollectionAlgo.h"
 #include "EdModeInteractiveToolsContext.h"
 #include "BaseGizmos/TransformGizmoUtil.h"
+#include "Snapping/ModelingSceneSnappingManager.h"
 #include "Components/InstancedStaticMeshComponent.h"
 
 
@@ -63,6 +64,9 @@ void UFractureEditorMode::Enter()
 	FCoreUObjectDelegates::OnPackageReloaded.AddUObject(this, &UFractureEditorMode::HandlePackageReloaded);
 
 	UE::TransformGizmoUtil::RegisterTransformGizmoContextObject(GetInteractiveToolsContext());
+
+	// register snapping manager
+	UE::Geometry::RegisterSceneSnappingManager(GetInteractiveToolsContext());
 	
 	// Get initial geometry component selection from currently selected actors when we enter the mode
 	USelection* SelectedActors = GEditor->GetSelectedActors();
@@ -88,6 +92,8 @@ void UFractureEditorMode::Exit()
 	// will be called before our Exit(); add the below line back after this bug is fixed
 	//FractureGizmoHelper->DeregisterGizmosWithManager(ToolsContext->ToolManager);
 
+	UE::Geometry::DeregisterSceneSnappingManager(GetInteractiveToolsContext());
+
 	// Empty the geometry component selection set
 	TArray<UObject*> SelectedObjects;
 	OnActorSelectionChanged(SelectedObjects, false);
diff --git a/Engine/Plugins/Experimental/ChaosEditor/Source/FractureEditor/Private/FractureEditorModeToolkit.cpp b/Engine/Plugins/Experimental/ChaosEditor/Source/FractureEditor/Private/FractureEditorModeToolkit.cpp
--- a/Engine/Plugins/Experimental/ChaosEditor/Source/FractureEditor/Private/FractureEditorModeToolkit.cpp
+++ b/Engine/Plugins/Experimental/ChaosEditor/Source/FractureEditor/Private/FractureEditorModeToolkit.cpp
@@ -1284,6 +1284,8 @@ void FFractureEditorModeToolkit::ShutdownActiveTool()
 	{
 		ActiveTool->Shutdown();
 		ActiveTool->OnPropertyModifiedDirectlyByTool.RemoveAll(this);
+
+		ActiveTool = nullptr;
 	}
 }
 
diff --git a/Engine/Plugins/Experimental/ChaosEditor/Source/FractureEditor/Private/FractureToolCutter.cpp b/Engine/Plugins/Experimental/ChaosEditor/Source/FractureEditor/Private/FractureToolCutter.cpp
--- a/Engine/Plugins/Experimental/ChaosEditor/Source/FractureEditor/Private/FractureToolCutter.cpp
+++ b/Engine/Plugins/Experimental/ChaosEditor/Source/FractureEditor/Private/FractureToolCutter.cpp
@@ -6,8 +6,10 @@
 #include "FractureToolContext.h"
 #include "InteractiveToolsContext.h"
 #include "EditorModeManager.h"
+#include "FractureEditorMode.h"
 #include "ToolSetupUtil.h"
 #include "BaseGizmos/GizmoBaseComponent.h"
+#include "BaseGizmos/TransformGizmoUtil.h"
 
 #include "GeometryCollection/GeometryCollectionObject.h"
 #include "GeometryCollection/GeometryCollectionComponent.h"
@@ -76,12 +78,13 @@ void UFractureTransformGizmoSettings::TransformChanged(UTransformProxy* Proxy, F
 void UFractureTransformGizmoSettings::Setup(UFractureToolCutterBase* Cutter, ETransformGizmoSubElements GizmoElements)
 {
 	AttachedCutter = Cutter;
-	UInteractiveToolsContext* Context = GLevelEditorModeTools().GetInteractiveToolsContext();
-	if (ensure(Context && AttachedCutter))
+	UFractureEditorMode* Mode = Cast<UFractureEditorMode>(GLevelEditorModeTools().GetActiveScriptableMode(UFractureEditorMode::EM_FractureEditorModeId));
+	UsedToolsContext = Mode->GetInteractiveToolsContext();
+	if (ensure(UsedToolsContext && AttachedCutter))
 	{
-		UInteractiveGizmoManager* GizmoManager = Context->GizmoManager;
+		UInteractiveGizmoManager* GizmoManager = UsedToolsContext->GizmoManager;
 		TransformProxy = NewObject<UTransformProxy>(this);
-		TransformGizmo = GizmoManager->CreateCustomTransformGizmo(GizmoElements, this);
+		TransformGizmo = UE::TransformGizmoUtil::CreateCustomTransformGizmo(GizmoManager, GizmoElements, this);
 		// TODO: Stop setting bUseEditorCompositing on gizmo components like this, once gizmo rendering is improved
 		// This is a hack to make gizmos more visible when the translucent fracture bone selection material is on screen
 		// This hack unfortunately makes the gizmos dithered when they're behind objects, but at least they're relatively visible
@@ -102,10 +105,10 @@ void UFractureTransformGizmoSettings::Setup(UFractureToolCutterBase* Cutter, ETr
 
 void UFractureTransformGizmoSettings::Shutdown()
 {
-	UInteractiveToolsContext* Context = GLevelEditorModeTools().GetInteractiveToolsContext();
-	if (Context)
+	if (UsedToolsContext)
 	{
-		Context->GizmoManager->DestroyAllGizmosByOwner(this);
+		UsedToolsContext->GizmoManager->DestroyAllGizmosByOwner(this);
+		UsedToolsContext = nullptr;
 	}
 	TransformGizmo = nullptr;
 	TransformProxy = nullptr;
diff --git a/Engine/Plugins/Experimental/ChaosEditor/Source/FractureEditor/Private/FractureToolCutter.h b/Engine/Plugins/Experimental/ChaosEditor/Source/FractureEditor/Private/FractureToolCutter.h
--- a/Engine/Plugins/Experimental/ChaosEditor/Source/FractureEditor/Private/FractureToolCutter.h
+++ b/Engine/Plugins/Experimental/ChaosEditor/Source/FractureEditor/Private/FractureToolCutter.h
@@ -304,6 +304,11 @@ public:
 
 protected:
 
+	UPROPERTY()
 	TObjectPtr<UFractureToolCutterBase> AttachedCutter = nullptr;
+
+	// the tools context responsible for the gizmo
+	UPROPERTY()
+	TObjectPtr<UInteractiveToolsContext> UsedToolsContext = nullptr;
 };
 
diff --git a/Engine/Plugins/Experimental/ChaosEditor/Source/FractureEditor/Private/FractureToolSelection.cpp b/Engine/Plugins/Experimental/ChaosEditor/Source/FractureEditor/Private/FractureToolSelection.cpp
--- a/Engine/Plugins/Experimental/ChaosEditor/Source/FractureEditor/Private/FractureToolSelection.cpp
+++ b/Engine/Plugins/Experimental/ChaosEditor/Source/FractureEditor/Private/FractureToolSelection.cpp
@@ -73,6 +73,12 @@ void UFractureToolSelection::RegisterUICommand( FFractureEditorCommands* Binding
 
 void UFractureToolSelection::CreateRectangleSelectionBehavior()
 {
+	UFractureEditorMode* Mode = Cast<UFractureEditorMode>(GLevelEditorModeTools().GetActiveScriptableMode(UFractureEditorMode::EM_FractureEditorModeId));
+	if (!Mode)
+	{
+		return;
+	}
+
 	// Set things up for being able to add behaviors to live preview.
 	SelectionBehaviorSet = NewObject<UInputBehaviorSet>();
 	SelectionBehaviorSource = NewObject<ULocalInputBehaviorSource>();
@@ -94,18 +100,18 @@ void UFractureToolSelection::CreateRectangleSelectionBehavior()
 	ClickOrDragBehavior->Modifiers.RegisterModifier(CtrlModifierID, FInputDeviceState::IsCtrlKeyDown);
 	SelectionBehaviorSet->Add(ClickOrDragBehavior, this);
 
-	UInteractiveToolsContext* Context = GLevelEditorModeTools().GetInteractiveToolsContext();
-	Context->InputRouter->RegisterSource(SelectionBehaviorSource);
+	UsedToolsContext = Mode->GetInteractiveToolsContext();
+	UsedToolsContext->InputRouter->RegisterSource(SelectionBehaviorSource);
 }
 
 void UFractureToolSelection::DestroyRectangleSelectionBehavior()
 {
-	if (SelectionBehaviorSource)
+	if (SelectionBehaviorSource && UsedToolsContext)
 	{
-		UInteractiveToolsContext* Context = GLevelEditorModeTools().GetInteractiveToolsContext();
-		Context->InputRouter->DeregisterSource(SelectionBehaviorSource);
+		UsedToolsContext->InputRouter->DeregisterSource(SelectionBehaviorSource);
 	}
 
+	UsedToolsContext = nullptr;
 	SelectionBehaviorSet = nullptr;
 	SelectionBehaviorSource = nullptr;
 	RectangleMarqueeManager = nullptr;
@@ -139,7 +145,6 @@ void UFractureToolSelection::OnDragRectangleChanged(const FCameraRectangle& Rect
 
 void UFractureToolSelection::OnClicked(const FInputDeviceRay& ClickPos)
 {
-	UInteractiveToolsContext* Context = GLevelEditorModeTools().GetInteractiveToolsContext();
 	UFractureEditorMode* Mode = Cast<UFractureEditorMode>(GLevelEditorModeTools().GetActiveScriptableMode(UFractureEditorMode::EM_FractureEditorModeId));
 	if (!Mode)
 	{
@@ -313,6 +318,7 @@ void URectangleMarqueeManager::DrawHUD(FCanvas* Canvas, bool bThisViewHasFocus)
 void UFractureToolSelection::Shutdown()
 {
 	DestroyRectangleSelectionBehavior();
+	Super::Shutdown();
 }
 
 void UFractureToolSelection::DrawHUD(FEditorViewportClient* ViewportClient, FViewport* Viewport, const FSceneView* View, FCanvas* Canvas)
diff --git a/Engine/Plugins/Experimental/ChaosEditor/Source/FractureEditor/Private/FractureToolSelection.h b/Engine/Plugins/Experimental/ChaosEditor/Source/FractureEditor/Private/FractureToolSelection.h
--- a/Engine/Plugins/Experimental/ChaosEditor/Source/FractureEditor/Private/FractureToolSelection.h
+++ b/Engine/Plugins/Experimental/ChaosEditor/Source/FractureEditor/Private/FractureToolSelection.h
@@ -248,6 +248,10 @@ protected:
 	UPROPERTY()
 	TObjectPtr<URectangleMarqueeManager> RectangleMarqueeManager = nullptr;
 
+	// the tools context responsible for the selection input router
+	UPROPERTY()
+	TObjectPtr<UInteractiveToolsContext> UsedToolsContext = nullptr;
+
 	// Support for Shift and Ctrl toggle
 	bool bShiftToggle = false;
 	bool bCtrlToggle = false;
