Chaos : Fracture Editor
- Guard against potential index bounds error.

diff --git a/Engine/Plugins/Experimental/ChaosEditor/Source/FractureEditor/Private/SGeometryCollectionOutliner.cpp b/Engine/Plugins/Experimental/ChaosEditor/Source/FractureEditor/Private/SGeometryCollectionOutliner.cpp
--- a/Engine/Plugins/Experimental/ChaosEditor/Source/FractureEditor/Private/SGeometryCollectionOutliner.cpp
+++ b/Engine/Plugins/Experimental/ChaosEditor/Source/FractureEditor/Private/SGeometryCollectionOutliner.cpp
@@ -370,7 +370,11 @@ void SGeometryCollectionOutliner::SetInitialDynamicState(int32 InDynamicState)
 				{
 					TManagedArray<int32>& InitialDynamicState = GeometryCollection->GetAttribute<int32>("InitialDynamicState", FGeometryCollection::TransformGroup);
 					TArray<int32> SelectedBones = Component->GetSelectedBones();
-					InitialDynamicState[SelectedItem->GetBoneIndex()] = InDynamicState;
+					int32 BoneIndex = SelectedItem->GetBoneIndex();
+					if (ensure(0 <= BoneIndex && BoneIndex < InitialDynamicState.Num()))
+					{
+						InitialDynamicState[SelectedItem->GetBoneIndex()] = InDynamicState;
+					}
 				}
 			}
 		}
