PR #8946: Minor non-unity build fixes (Contributed by [REDACTED_PERSON])

diff --git a/Engine/Plugins/Experimental/ChaosEditor/Source/FractureEditor/Private/FractureToolMeshCut.cpp b/Engine/Plugins/Experimental/ChaosEditor/Source/FractureEditor/Private/FractureToolMeshCut.cpp
--- a/Engine/Plugins/Experimental/ChaosEditor/Source/FractureEditor/Private/FractureToolMeshCut.cpp
+++ b/Engine/Plugins/Experimental/ChaosEditor/Source/FractureEditor/Private/FractureToolMeshCut.cpp
@@ -11,6 +11,9 @@
 #include "Engine/StaticMesh.h"
 #include "Misc/ScopedSlowTask.h"
 
+#include "Engine/StaticMesh.h"
+#include "Components/StaticMeshComponent.h"
+
 #define LOCTEXT_NAMESPACE "FractureMesh"
 
 
diff --git a/Engine/Plugins/Experimental/ChaosEditor/Source/FractureEditor/Public/FractureEditorModeToolkit.h b/Engine/Plugins/Experimental/ChaosEditor/Source/FractureEditor/Public/FractureEditorModeToolkit.h
--- a/Engine/Plugins/Experimental/ChaosEditor/Source/FractureEditor/Public/FractureEditorModeToolkit.h
+++ b/Engine/Plugins/Experimental/ChaosEditor/Source/FractureEditor/Public/FractureEditorModeToolkit.h
@@ -24,6 +24,7 @@ class UGeometryCollection;
 class FFractureEditorModeToolkit;
 class UFractureActionTool;
 class UFractureModalTool;
+enum class EMapChangeType : uint8;
 
 namespace GeometryCollection
 {
diff --git a/Engine/Source/Programs/HeadlessChaosPerf/Private/ChaosPerf/ChaosPerfTest.h b/Engine/Source/Programs/HeadlessChaosPerf/Private/ChaosPerf/ChaosPerfTest.h
--- a/Engine/Source/Programs/HeadlessChaosPerf/Private/ChaosPerf/ChaosPerfTest.h
+++ b/Engine/Source/Programs/HeadlessChaosPerf/Private/ChaosPerf/ChaosPerfTest.h
@@ -2,6 +2,7 @@
 #pragma once
 
 #include "ChaosPerf/ChaosPerf.h"
+#include "Misc/Paths.h"
 
 namespace ChaosPerf
 {
diff --git a/Engine/Source/Runtime/Engine/Private/SkeletalMeshComponentPhysics.cpp b/Engine/Source/Runtime/Engine/Private/SkeletalMeshComponentPhysics.cpp
--- a/Engine/Source/Runtime/Engine/Private/SkeletalMeshComponentPhysics.cpp
+++ b/Engine/Source/Runtime/Engine/Private/SkeletalMeshComponentPhysics.cpp
@@ -3,6 +3,7 @@
 #include "CoreMinimal.h"
 #include "Misc/ConfigCacheIni.h"
 #include "Misc/MessageDialog.h"
+#include "Misc/ConfigCacheIni.h"
 #include "Stats/Stats.h"
 #include "UObject/UObjectBaseUtility.h"
 #include "HAL/IConsoleManager.h"
diff --git a/Engine/Source/Runtime/Engine/Public/TextureResource.h b/Engine/Source/Runtime/Engine/Public/TextureResource.h
--- a/Engine/Source/Runtime/Engine/Public/TextureResource.h
+++ b/Engine/Source/Runtime/Engine/Public/TextureResource.h
@@ -32,6 +32,7 @@ class FTexture2DResource;
 class FTexture3DResource;
 class FTexture2DArrayResource;
 class FVirtualTexture2DResource;
+struct IPooledRenderTarget;
 
 /** Maximum number of slices in texture source art. */
 #define MAX_TEXTURE_SOURCE_SLICES 6
