Chaos : fix kinematic state not been properly serialized when set via the context menu

diff --git a/Engine/Plugins/Experimental/ChaosEditor/Source/FractureEditor/Private/SGeometryCollectionOutliner.cpp b/Engine/Plugins/Experimental/ChaosEditor/Source/FractureEditor/Private/SGeometryCollectionOutliner.cpp
--- a/Engine/Plugins/Experimental/ChaosEditor/Source/FractureEditor/Private/SGeometryCollectionOutliner.cpp
+++ b/Engine/Plugins/Experimental/ChaosEditor/Source/FractureEditor/Private/SGeometryCollectionOutliner.cpp
@@ -364,9 +364,11 @@ void SGeometryCollectionOutliner::SetInitialDynamicState(int32 InDynamicState)
 	{
 		if (UGeometryCollectionComponent* Component = SelectedItem->GetComponent())
 		{
-			if (const UGeometryCollection* RestCollection = Component->GetRestCollection())
+			FGeometryCollectionEdit GeometryCollectionEdit = Component->EditRestCollection();
+			if (UGeometryCollection* GeometryCollectionObject = GeometryCollectionEdit.GetRestCollection())
 			{
-				if (FGeometryCollection* GeometryCollection = RestCollection->GetGeometryCollection().Get())
+				TSharedPtr<FGeometryCollection,ESPMode::ThreadSafe> GeometryCollection = GeometryCollectionObject->GetGeometryCollection();
+				if (GeometryCollection)
 				{
 					TManagedArray<int32>& InitialDynamicState = GeometryCollection->GetAttribute<int32>("InitialDynamicState", FGeometryCollection::TransformGroup);
 					TArray<int32> SelectedBones = Component->GetSelectedBones();
