fix mesh fractures fracturing 'inactive,' fully-invisible geometry that was left on clusters for scatter/grid fractures

diff --git a/Engine/Plugins/Experimental/ChaosEditor/Source/FractureEditor/Private/FractureToolMeshCut.cpp b/Engine/Plugins/Experimental/ChaosEditor/Source/FractureEditor/Private/FractureToolMeshCut.cpp
--- a/Engine/Plugins/Experimental/ChaosEditor/Source/FractureEditor/Private/FractureToolMeshCut.cpp
+++ b/Engine/Plugins/Experimental/ChaosEditor/Source/FractureEditor/Private/FractureToolMeshCut.cpp
@@ -10,6 +10,7 @@
 #include "Components/StaticMeshComponent.h"
 #include "Engine/StaticMesh.h"
 #include "Misc/ScopedSlowTask.h"
+#include "Algo/RemoveIf.h"
 
 #include "Engine/StaticMesh.h"
 #include "Components/StaticMeshComponent.h"
@@ -224,7 +225,7 @@ int32 UFractureToolMeshCut::ExecuteFracture(const FFractureToolContext& Fracture
 			GenerateMeshTransforms(FractureContext, MeshTransforms);
 			int32 FirstIndex = -1;
 			// Create progress indicator dialog
-			static const FText SlowTaskText = LOCTEXT("CutWithScatteredMeshes", "Cutting geometry collection with plane(s)...");
+			static const FText SlowTaskText = LOCTEXT("CutWithScatteredMeshes", "Cutting geometry collection with mesh ...");
 
 			FScopedSlowTask SlowTask(MeshTransforms.Num(), SlowTaskText);
 			SlowTask.MakeDialog();
@@ -242,6 +243,11 @@ int32 UFractureToolMeshCut::ExecuteFracture(const FFractureToolContext& Fracture
 			{
 				EnterProgressFrame(1);
 				int32 Index = CutWithMesh(MeshDescription, ScatterTransform, InternalSurfaceMaterials, *FractureContext.GetGeometryCollection(), BonesToCut, CollisionSettings->GetPointSpacing(), FractureContext.GetTransform());
+				int32 NewLen = Algo::RemoveIf(BonesToCut, [&FractureContext](int32 Bone)
+					{
+						return !FractureContext.GetGeometryCollection()->IsVisible(Bone); // remove already-fractured pieces from the to-cut list
+					});
+				BonesToCut.SetNum(NewLen);
 				if (FirstIndex == -1)
 				{
 					FirstIndex = Index;
