fix crash due to fracture mode outliner being out of sync with geometry after an undo or redo + add some code to prevent crashes if this kind of invalid selection occurs again

diff --git a/Engine/Plugins/Experimental/ChaosEditor/Source/FractureEditor/Private/FractureEditorMode.cpp b/Engine/Plugins/Experimental/ChaosEditor/Source/FractureEditor/Private/FractureEditorMode.cpp
--- a/Engine/Plugins/Experimental/ChaosEditor/Source/FractureEditor/Private/FractureEditorMode.cpp
+++ b/Engine/Plugins/Experimental/ChaosEditor/Source/FractureEditor/Private/FractureEditorMode.cpp
@@ -120,8 +120,10 @@ bool UFractureEditorMode::MatchesContext(const FTransactionContext& InContext, c
 	{
 		return true;
 	}
-
-	return false;
+	
+	// TODO: Rigorously use transaction contexts for fracture-related transactions, so we can filter them here
+	// Once that's done, we can detect the fracture contexts above and change this last line back to a "return false;"
+	return true;
 }
 
 void UFractureEditorMode::PostUndo(bool bSuccess)
@@ -438,6 +440,7 @@ bool UFractureEditorMode::GetPivotForOrbit(FVector& OutPivot) const
 
 void UFractureEditorMode::OnUndoRedo()
 {
+	RefreshOutlinerWithCurrentSelection(); // always refresh the outliner in case the geometry collection bones have changed
 	for (UGeometryCollectionComponent* SelectedComp : SelectedGeometryComponents)
 	{
 		// We need to update the bone colors to account for undoing/redoing selection
diff --git a/Engine/Plugins/Experimental/ChaosEditor/Source/FractureEditor/Private/FractureEditorModeToolkit.cpp b/Engine/Plugins/Experimental/ChaosEditor/Source/FractureEditor/Private/FractureEditorModeToolkit.cpp
--- a/Engine/Plugins/Experimental/ChaosEditor/Source/FractureEditor/Private/FractureEditorModeToolkit.cpp
+++ b/Engine/Plugins/Experimental/ChaosEditor/Source/FractureEditor/Private/FractureEditorModeToolkit.cpp
@@ -72,6 +72,8 @@
 #include "Toolkits/AssetEditorModeUILayer.h"
 #include "SPrimaryButton.h"
 
+#include "Algo/RemoveIf.h"
+
 #define LOCTEXT_NAMESPACE "FFractureEditorModeToolkit"
 
 TArray<UClass*> FindFractureToolClasses()
@@ -1377,7 +1379,7 @@ void FFractureEditorModeToolkit::SetOutlinerComponents(const TArray<UGeometryCol
 		FGeometryCollectionEdit RestCollection = Component->EditRestCollection(GeometryCollection::EEditUpdate::None);
 		UGeometryCollection* FracturedGeometryCollection = RestCollection.GetRestCollection();
 
-		if (FracturedGeometryCollection && IsValidChecked(FracturedGeometryCollection)) // Prevents crash when GC is deleted from content browser and actor is selected.
+		if (IsValid(FracturedGeometryCollection)) // Prevents crash when GC is deleted from content browser and actor is selected.
 		{
 			TSharedPtr<FGeometryCollection, ESPMode::ThreadSafe> GeometryCollectionPtr = FracturedGeometryCollection->GetGeometryCollection();
 
@@ -1732,6 +1734,14 @@ void FFractureEditorModeToolkit::OnOutlinerBoneSelectionChanged(UGeometryCollect
 	const UGeometryCollection* RestCollection = RootComponent->GetRestCollection();
 	if (RestCollection && IsValidChecked(RestCollection))
 	{
+		int32 NumValidBones = Algo::RemoveIf(SelectedBones, [&RestCollection](const int32& Bone)
+			{
+				return Bone < 0 || Bone >= RestCollection->GetGeometryCollection()->NumElements(FGeometryCollection::TransformGroup);
+			});
+		if (!ensure(NumValidBones == SelectedBones.Num())) // Protect against invalid bones in selection, but ensure() as this indicates the UI is out of sync with the data
+		{
+			SelectedBones.SetNum(NumValidBones);
+		}
 		if (SelectedBones.Num())
 		{
 			// don't need to snap the bones to the current level because they are directly selected from the outliner
diff --git a/Engine/Plugins/Experimental/ChaosEditor/Source/FractureEditor/Private/SGeometryCollectionOutliner.cpp b/Engine/Plugins/Experimental/ChaosEditor/Source/FractureEditor/Private/SGeometryCollectionOutliner.cpp
--- a/Engine/Plugins/Experimental/ChaosEditor/Source/FractureEditor/Private/SGeometryCollectionOutliner.cpp
+++ b/Engine/Plugins/Experimental/ChaosEditor/Source/FractureEditor/Private/SGeometryCollectionOutliner.cpp
@@ -595,20 +595,22 @@ TSharedRef<ITableRow> FGeometryCollectionTreeItemBone::MakeTreeRowWidget(const T
 	{
 		TSharedPtr<FGeometryCollection, ESPMode::ThreadSafe> GeometryCollectionPtr = RestCollection->GetGeometryCollection();
 		const TManagedArray<int32>& SimulationType = GeometryCollectionPtr->SimulationType;
-		switch (SimulationType[GetBoneIndex()])
+		if (ensure(GetBoneIndex() >= 0 && GetBoneIndex() < SimulationType.Num()))
 		{
+			switch (SimulationType[GetBoneIndex()])
+			{
 			case FGeometryCollection::ESimulationTypes::FST_None:
 				TextColor = FLinearColor::Green;
 				break;
 
 			case FGeometryCollection::ESimulationTypes::FST_Rigid:
 				if (GeometryCollectionPtr->IsVisible(GetBoneIndex()))
-				{ 
+				{
 					TextColor = FSlateColor::UseForeground();
 				}
 				else
 				{
-					TextColor = FLinearColor{0.1f, 0.1f, 0.1f, 1.f};
+					TextColor = FLinearColor{ 0.1f, 0.1f, 0.1f, 1.f };
 				}
 				break;
 
@@ -619,9 +621,16 @@ TSharedRef<ITableRow> FGeometryCollectionTreeItemBone::MakeTreeRowWidget(const T
 			default:
 				ensureMsgf(false, TEXT("Invalid Geometry Collection simulation type encountered."));
 				break;
-		}
+			}
 
-		InitialDynamicState = GeometryCollectionPtr->InitialDynamicState[GetBoneIndex()];
+			InitialDynamicState = GeometryCollectionPtr->InitialDynamicState[GetBoneIndex()];
+		}
+		else
+		{
+			// UI contains invalid bone indices; likely not correctly updated after an undo or redo?
+			TextColor = FLinearColor(0.1f, 0.1f, 0.1f);
+			InitialDynamicState = INDEX_NONE;
+		}
 	}
 	else
 	{
