Add per level coloring option in fracture editor outlinerand stats details

diff --git a/Engine/Plugins/Experimental/ChaosEditor/Source/FractureEditor/Private/FractureEditorModeToolkit.cpp b/Engine/Plugins/Experimental/ChaosEditor/Source/FractureEditor/Private/FractureEditorModeToolkit.cpp
--- a/Engine/Plugins/Experimental/ChaosEditor/Source/FractureEditor/Private/FractureEditorModeToolkit.cpp
+++ b/Engine/Plugins/Experimental/ChaosEditor/Source/FractureEditor/Private/FractureEditorModeToolkit.cpp
@@ -50,6 +50,7 @@
 #include "FractureToolAutoCluster.h" 
 #include "SGeometryCollectionOutliner.h"
 #include "SGeometryCollectionHistogram.h"
+#include "SGeometryCollectionStatistics.h"
 #include "FractureSelectionTools.h"
 #include "GeometryCollection/GeometryCollectionAlgo.h"
 #include "GeometryCollection/GeometryCollectionClusteringUtility.h"
@@ -633,8 +634,9 @@ TSharedRef<SDockTab> FFractureEditorModeToolkit::CreateStatisticsTab(const FSpaw
 		.AreaTitleFont(FEditorStyle::Get().GetFontStyle("DetailsView.CategoryFontStyle"))
 		.BodyContent()
 		[
-			SNew(STextBlock)
-			.Text(this, &FFractureEditorModeToolkit::GetStatisticsSummary)
+			SAssignNew(StatisticsView, SGeometryCollectionStatistics)
+			//SNew(STextBlock)
+			//.Text(LOCTEXT("Statt", "Statt")) //this, &FFractureEditorModeToolkit::GetStatisticsSummary)
 		];
 	TSharedPtr<SDockTab> CreatedTab = SNew(SDockTab)
 		[
@@ -682,7 +684,14 @@ void FFractureEditorModeToolkit::OnObjectPostEditChange( UObject* Object, FPrope
 			UOutlinerSettings* OutlinerSettings = GetMutableDefault<UOutlinerSettings>();
 			OutlinerView->RegenerateItems();
 		}
-
+		else if (PropertyChangedEvent.Property->GetFName() == GET_MEMBER_NAME_CHECKED(UOutlinerSettings, ColorByLevel))
+		{
+			UOutlinerSettings* OutlinerSettings = GetMutableDefault<UOutlinerSettings>();
+			OutlinerView->RegenerateItems();
+			FGeometryCollectionStatistics Stats;
+			GetStatisticsSummary(Stats);
+			StatisticsView->SetStatistics(Stats);
+		}
 	}
 }
 
@@ -1404,6 +1413,13 @@ void FFractureEditorModeToolkit::SetOutlinerComponents(const TArray<UGeometryCol
 		HistogramView->SetComponents(ComponentsToEdit, GetLevelViewValue());
 	}
 
+	if (StatisticsView)
+	{
+		FGeometryCollectionStatistics Stats;
+		GetStatisticsSummary(Stats);
+		StatisticsView->SetStatistics(Stats);
+	}
+
 	if (ActiveTool != nullptr)
 	{
 		ActiveTool->SelectedBonesChanged();
@@ -1805,9 +1821,8 @@ FText FFractureEditorModeToolkit::GetSelectionInfo() const
 }
 
 
-FText FFractureEditorModeToolkit::GetStatisticsSummary() const
+void FFractureEditorModeToolkit::GetStatisticsSummary(FGeometryCollectionStatistics& Stats) const
 {
-
 	TArray<const FGeometryCollection*> GeometryCollectionArray;
 	if (USelection* SelectedActors = GEditor->GetSelectedActors())
 	{
@@ -1831,8 +1846,6 @@ FText FFractureEditorModeToolkit::GetStatisticsSummary() const
 	}
 
 
-	FString Buffer;
-
 	if (GeometryCollectionArray.Num() > 0)
 	{
 		TArray<int32> LevelTransformsAll;
@@ -1845,9 +1858,6 @@ FText FFractureEditorModeToolkit::GetStatisticsSummary() const
 
 			check(GeometryCollection);
 
-
-			Buffer += FString::Printf(TEXT("Sum of the selected Geometry Collections\n\n"));
-
 			if(GeometryCollection->HasAttribute("Level", FGeometryCollection::TransformGroup))
 			{
 				const TManagedArray<int32>& Levels = GeometryCollection->GetAttribute<int32>("Level", FGeometryCollection::TransformGroup);
@@ -1887,14 +1897,14 @@ FText FFractureEditorModeToolkit::GetStatisticsSummary() const
 			}
 		}
 
+		Stats.CountsPerLevel.Reset();
 		for (int32 Level = 0; Level < LevelMax; ++Level)
 		{
-			Buffer += FString::Printf(TEXT("Level: %d \t - \t %d\n"), Level, LevelTransformsAll[Level]);
+			Stats.CountsPerLevel.Add(LevelTransformsAll[Level]);
 		}
-		Buffer += FString::Printf(TEXT("\nEmbedded: %d"), EmbeddedCount);
-	}
 
-	return FText::FromString(Buffer);
+		Stats.EmbeddedCount = EmbeddedCount;
+	}
 }
 
 #undef LOCTEXT_NAMESPACE
diff --git a/Engine/Plugins/Experimental/ChaosEditor/Source/FractureEditor/Private/SGeometryCollectionOutliner.cpp b/Engine/Plugins/Experimental/ChaosEditor/Source/FractureEditor/Private/SGeometryCollectionOutliner.cpp
--- a/Engine/Plugins/Experimental/ChaosEditor/Source/FractureEditor/Private/SGeometryCollectionOutliner.cpp
+++ b/Engine/Plugins/Experimental/ChaosEditor/Source/FractureEditor/Private/SGeometryCollectionOutliner.cpp
@@ -73,6 +73,19 @@ void FGeometryCollectionTreeItem::GenerateContextMenu(UToolMenu* Menu, SGeometry
 	Section.AddSubMenu("FractureToolSetInitialDynamicStateMenu", NSLOCTEXT("Fracture", "FractureToolSetInitialDynamicStateMenu", "Initial Dynamic State"), FText(),FNewToolMenuDelegate::CreateLambda(MakeDynamicStateMenu));
 }
 
+FColor FGeometryCollectionTreeItem::GetColorPerDepth(uint32 Depth)
+{
+	const FColor ColorsPerDepth[] =
+	{
+		FColor::Cyan,
+		FColor::Orange,
+		FColor::Emerald,
+		FColor::Yellow,
+	};
+	const uint32 ColorPerDepthCount = sizeof(ColorsPerDepth) / sizeof(FColor);
+	return ColorsPerDepth[Depth % ColorPerDepthCount];
+}
+
 void FGeometryCollectionTreeItemBone::OnDragEnter(FDragDropEvent const& InDragDropEvent)
 {
 	TSharedPtr<FDragDropOperation> Operation = InDragDropEvent.GetOperation();
@@ -148,6 +161,7 @@ FReply FGeometryCollectionTreeItemBone::OnDrop(const FDragDropEvent& DragDropEve
 UOutlinerSettings::UOutlinerSettings(const FObjectInitializer& ObjInit)
 	: Super(ObjInit)
 	, ItemText(EOutlinerItemNameEnum::BoneIndex)
+	, ColorByLevel(false)
 {}
 
 
@@ -165,6 +179,8 @@ void SGeometryCollectionOutliner::Construct(const FArguments& InArgs)
 		.OnGetChildren(this, &SGeometryCollectionOutliner::OnGetChildren)
 		.OnContextMenuOpening(this, &SGeometryCollectionOutliner::OnOpenContextMenu)
 		.AllowInvisibleItemSelection(true)
+		.ShouldStackHierarchyHeaders(true)
+		.OnGeneratePinnedRow(this, &SGeometryCollectionOutliner::OnGeneratePinnedRowWidget, true)
 		.HighlightParentNodesForSelection(true)
 		.OnSetExpansionRecursive(this, &SGeometryCollectionOutliner::ExpandRecursive)
 	];
@@ -180,6 +196,11 @@ TSharedRef<ITableRow> SGeometryCollectionOutliner::MakeTreeRowWidget(FGeometryCo
 	return InItem->MakeTreeRowWidget(InOwnerTable);
 }
 
+TSharedRef<ITableRow> SGeometryCollectionOutliner::OnGeneratePinnedRowWidget(FGeometryCollectionTreeItemPtr InItem, const TSharedRef<STableViewBase>& InOwnerTable, bool bPinned)
+{
+	return InItem->MakeTreeRowWidget(InOwnerTable, true);
+}
+
 void SGeometryCollectionOutliner::OnGetChildren(FGeometryCollectionTreeItemPtr InItem, TArray<FGeometryCollectionTreeItemPtr>& OutChildren)
 {
 	InItem->GetChildren(OutChildren);
@@ -384,7 +405,7 @@ void SGeometryCollectionOutliner::SetInitialDynamicState(int32 InDynamicState)
 	RegenerateItems();
 }
 
-TSharedRef<ITableRow> FGeometryCollectionTreeItemComponent::MakeTreeRowWidget(const TSharedRef<STableViewBase>& InOwnerTable)
+TSharedRef<ITableRow> FGeometryCollectionTreeItemComponent::MakeTreeRowWidget(const TSharedRef<STableViewBase>& InOwnerTable, bool bNoExtraColumn)
 {
 	FString ActorName = Component->GetOwner()->GetActorLabel();
 	FString ComponentName = Component->GetClass()->GetFName().ToString();
@@ -576,7 +597,7 @@ bool FGeometryCollectionTreeItemComponent::FilterBoneIndex(int32 BoneIndex) cons
 	return true;	
 }
 
-TSharedRef<ITableRow> FGeometryCollectionTreeItemBone::MakeTreeRowWidget(const TSharedRef<STableViewBase>& InOwnerTable)
+TSharedRef<ITableRow> FGeometryCollectionTreeItemBone::MakeTreeRowWidget(const TSharedRef<STableViewBase>& InOwnerTable, bool bNoExtraColumn)
 {
 	UOutlinerSettings* OutlinerSettings = GetMutableDefault<UOutlinerSettings>();
 	FText ItemText = ParentComponentItem->GetDisplayNameForBone(Guid);
@@ -597,6 +618,14 @@ TSharedRef<ITableRow> FGeometryCollectionTreeItemBone::MakeTreeRowWidget(const T
 		const TManagedArray<int32>& SimulationType = GeometryCollectionPtr->SimulationType;
 		if (ensure(GetBoneIndex() >= 0 && GetBoneIndex() < SimulationType.Num()))
 		{
+			const bool bHasLevelAttribute = GeometryCollectionPtr->HasAttribute("Level", FGeometryCollection::TransformGroup);
+			if (bHasLevelAttribute && OutlinerSettings->ColorByLevel)
+			{
+				TManagedArray<int32>& Level = GeometryCollectionPtr->GetAttribute<int32>("Level", FTransformCollection::TransformGroup);
+				TextColor = FSlateColor(FGeometryCollectionTreeItem::GetColorPerDepth((uint32)Level[GetBoneIndex()]));
+			}
+			else
+			{
 			switch (SimulationType[GetBoneIndex()])
 			{
 			case FGeometryCollection::ESimulationTypes::FST_None:
@@ -622,6 +651,7 @@ TSharedRef<ITableRow> FGeometryCollectionTreeItemBone::MakeTreeRowWidget(const T
 				ensureMsgf(false, TEXT("Invalid Geometry Collection simulation type encountered."));
 				break;
 			}
+			}
 
 			InitialDynamicState = GeometryCollectionPtr->InitialDynamicState[GetBoneIndex()];
 		}
@@ -639,7 +669,10 @@ TSharedRef<ITableRow> FGeometryCollectionTreeItemBone::MakeTreeRowWidget(const T
 		InitialDynamicState = INDEX_NONE;
 	}
 	
-	return SNew(STableRow<FGeometryCollectionTreeItemPtr>, InOwnerTable)
+	if (bNoExtraColumn)
+	{
+		// short version 
+		return SNew(STableRow<FGeometryCollectionTreeItemPtr>, InOwnerTable)
 		.Content()
 		[
 			SNew(SHorizontalBox)
@@ -651,14 +684,34 @@ TSharedRef<ITableRow> FGeometryCollectionTreeItemBone::MakeTreeRowWidget(const T
 					.Text(ItemText)
 					.ColorAndOpacity(TextColor)
 				]
+		]
+		.OnDragDetected(this, &FGeometryCollectionTreeItem::OnDragDetected)
+		.OnDrop(this, &FGeometryCollectionTreeItem::OnDrop)
+		.OnDragEnter(this, &FGeometryCollectionTreeItem::OnDragEnter)
+		.OnDragLeave(this, &FGeometryCollectionTreeItem::OnDragLeave);
+	}
+	// full item 
+	return SNew(STableRow<FGeometryCollectionTreeItemPtr>, InOwnerTable)
+		.Content()
+		[
+			SNew(SHorizontalBox)
 			+ SHorizontalBox::Slot()
-				.Padding(2.0f, 2.0f)
-				.FillWidth(1.0f)
-				.HAlign(HAlign_Right)
+				.Padding(2.0f, 4.0f)
+				.AutoWidth()
 				[
 					SNew(STextBlock)
-					.Text(GettextFromInitialDynamicState(InitialDynamicState))
+					.Text(ItemText)
+					.ColorAndOpacity(TextColor)
 				]
+			+ SHorizontalBox::Slot()
+					.Padding(2.0f, 2.0f)
+					.FillWidth(1.0f)
+					.HAlign(HAlign_Right)
+					[
+						SNew(STextBlock)
+						.Text(GettextFromInitialDynamicState(InitialDynamicState))
+						.ColorAndOpacity(TextColor)
+					]
 		]
 		.OnDragDetected(this, &FGeometryCollectionTreeItem::OnDragDetected)
 		.OnDrop(this, &FGeometryCollectionTreeItem::OnDrop)
diff --git a/Engine/Plugins/Experimental/ChaosEditor/Source/FractureEditor/Private/SGeometryCollectionOutliner.h b/Engine/Plugins/Experimental/ChaosEditor/Source/FractureEditor/Private/SGeometryCollectionOutliner.h
--- a/Engine/Plugins/Experimental/ChaosEditor/Source/FractureEditor/Private/SGeometryCollectionOutliner.h
+++ b/Engine/Plugins/Experimental/ChaosEditor/Source/FractureEditor/Private/SGeometryCollectionOutliner.h
@@ -47,6 +47,10 @@ public:
 	/** What is displayed in Outliner text */
 	UPROPERTY(EditAnywhere, Category = OutlinerSettings, meta = (DisplayName = "Item Text"))
 	EOutlinerItemNameEnum ItemText;
+
+	/** whther to use level coloring */
+	UPROPERTY(EditAnywhere, Category = OutlinerSettings, meta = (DisplayName = "Color By Level"))
+	bool ColorByLevel;
 };
 
 
@@ -54,7 +58,7 @@ class FGeometryCollectionTreeItem : public TSharedFromThis<FGeometryCollectionTr
 {
 public:
 	virtual ~FGeometryCollectionTreeItem() {}
-	virtual TSharedRef<ITableRow> MakeTreeRowWidget(const TSharedRef<STableViewBase>& InOwnerTable) = 0;
+	virtual TSharedRef<ITableRow> MakeTreeRowWidget(const TSharedRef<STableViewBase>& InOwnerTable, bool bNoExtraColumn = false) = 0;
 	virtual void GetChildren(FGeometryCollectionTreeItemList& OutChildren) = 0;
 	virtual UGeometryCollectionComponent* GetComponent() const = 0;
 
@@ -68,6 +72,8 @@ public:
 
 	//void MakeDynamicStateMenu(UToolMenu* Menu, SGeometryCollectionOutliner& Outliner);
 	void GenerateContextMenu(UToolMenu* Menu, SGeometryCollectionOutliner& Outliner);
+
+	static FColor GetColorPerDepth(uint32 Depth);
 };
 
 class FGeometryCollectionTreeItemComponent : public FGeometryCollectionTreeItem
@@ -81,7 +87,7 @@ public:
 	}
 
 	/** FGeometryCollectionTreeItem interface */
-	virtual TSharedRef<ITableRow> MakeTreeRowWidget(const TSharedRef<STableViewBase>& InOwnerTable);
+	virtual TSharedRef<ITableRow> MakeTreeRowWidget(const TSharedRef<STableViewBase>& InOwnerTable, bool bNoExtraColumn = false);
 	virtual void GetChildren(FGeometryCollectionTreeItemList& OutChildren) override;
 	virtual UGeometryCollectionComponent* GetComponent() const override { return Component.Get(); }
 
@@ -127,7 +133,7 @@ public:
 	{}
 
 	/** FGeometryCollectionTreeItem interface */
-	TSharedRef<ITableRow> MakeTreeRowWidget(const TSharedRef<STableViewBase>& InOwnerTable);//, const TAttribute<FText>& InFilterText) = 0;
+	TSharedRef<ITableRow> MakeTreeRowWidget(const TSharedRef<STableViewBase>& InOwnerTable, bool bNoExtraColumn = false);
 	virtual void GetChildren(FGeometryCollectionTreeItemList& OutChildren) override;
 	virtual int32 GetBoneIndex() const override { return BoneIndex; }
 	virtual UGeometryCollectionComponent* GetComponent() const { return ParentComponentItem->GetComponent(); }
@@ -161,6 +167,7 @@ public:
 	void RegenerateItems();
 
 	TSharedRef<ITableRow> MakeTreeRowWidget(FGeometryCollectionTreeItemPtr InInfo, const TSharedRef<STableViewBase>& OwnerTable);
+	TSharedRef<ITableRow> OnGeneratePinnedRowWidget(FGeometryCollectionTreeItemPtr InItem, const TSharedRef<STableViewBase>& InOwnerTable, bool bPinned);
 	void OnGetChildren(TSharedPtr<FGeometryCollectionTreeItem> InInfo, TArray< TSharedPtr<FGeometryCollectionTreeItem> >& OutChildren);
 	TSharedPtr<SWidget> OnOpenContextMenu();
 
diff --git a/Engine/Plugins/Experimental/ChaosEditor/Source/FractureEditor/Private/SGeometryCollectionStatistics.cpp b/Engine/Plugins/Experimental/ChaosEditor/Source/FractureEditor/Private/SGeometryCollectionStatistics.cpp
new file mode 100644
--- /dev/null
+++ b/Engine/Plugins/Experimental/ChaosEditor/Source/FractureEditor/Private/SGeometryCollectionStatistics.cpp
@@ -0,0 +1,141 @@
+// Copyright Epic Games, Inc. All Rights Reserved.
+
+#include "SGeometryCollectionStatistics.h"
+#include "SGeometryCollectionOutliner.h"
+
+#include "Widgets/Layout/SBox.h"
+#include "SlateOptMacros.h"
+
+#define LOCTEXT_NAMESPACE "SGeometryCollectionStatistics"
+
+namespace GeometryCollectionStatisticsUI
+{
+	static const FName NameLabel(TEXT("Stat"));
+	static const FName ValueLabel(TEXT("Value"));
+}
+
+class SGeometryCollectionStatisticsRow : public SMultiColumnTableRow<TSharedPtr<int32>>
+{
+	SLATE_BEGIN_ARGS(SGeometryCollectionStatisticsRow) {}
+		SLATE_ARGUMENT(FText, Name)
+		SLATE_ARGUMENT(FText, Value)
+		SLATE_ARGUMENT(FSlateColor, TextColor)
+	SLATE_END_ARGS()
+
+public:
+
+	/**
+	 * Constructs the widget.
+	 *
+	 * @param InArgs The construction arguments.
+	 */
+	void Construct(const FArguments& InArgs, const TSharedRef<STableViewBase>& InOwnerTableView)
+	{
+		Name = InArgs._Name;
+		Value = InArgs._Value;
+		TextColor = InArgs._TextColor;
+		SMultiColumnTableRow<TSharedPtr<int32> >::Construct(FSuperRowType::FArguments(), InOwnerTableView);
+	}
+
+public:
+
+	BEGIN_SLATE_FUNCTION_BUILD_OPTIMIZATION
+	virtual TSharedRef<SWidget> GenerateWidgetForColumn(const FName& ColumnName) override
+	{
+		FText ColumnText;
+
+		if (ColumnName == GeometryCollectionStatisticsUI::NameLabel)
+		{
+			ColumnText = Name;
+		}
+		else if (ColumnName == GeometryCollectionStatisticsUI::ValueLabel)
+		{
+			ColumnText = Value;
+		}
+
+		return SNew(SBox)
+			.Padding(FMargin(4.0, 0.0))
+			[
+				SNew(STextBlock)
+				.ColorAndOpacity(TextColor)
+				.Text(ColumnText)
+			];
+	}
+	END_SLATE_FUNCTION_BUILD_OPTIMIZATION
+
+private:
+
+	/** Name of the stat. */
+	FText Name;
+
+	/** Value of the stat. */
+	FText Value;
+
+	FSlateColor TextColor;
+};
+
+void SGeometryCollectionStatistics::Construct(const FArguments& InArgs)
+{
+	ChildSlot
+		[
+			SAssignNew(StatListView, SListView<FListItemPtr>)
+			.ListItemsSource(&StatItems)
+			.OnGenerateRow(this, &SGeometryCollectionStatistics::HandleGenerateRow)
+			.HeaderRow
+			(
+				SNew(SHeaderRow)
+				+SHeaderRow::Column(GeometryCollectionStatisticsUI::NameLabel)
+				.DefaultLabel(LOCTEXT("NameColumnHeaderName", "Name"))
+				+SHeaderRow::Column(GeometryCollectionStatisticsUI::ValueLabel)
+				.DefaultLabel(LOCTEXT("ValueColumnHeaderName", "Value"))
+			)
+		];
+}
+
+void SGeometryCollectionStatistics::Refresh()
+{
+	StatListView->RequestListRefresh();
+}
+
+void SGeometryCollectionStatistics::SetStatistics(const FGeometryCollectionStatistics& Stats)
+{
+	UOutlinerSettings* OutlinerSettings = GetMutableDefault<UOutlinerSettings>();
+	StatItems.Reset();
+	
+	for (int32 Level = 0; Level < Stats.CountsPerLevel.Num(); ++Level)
+	{
+		FSlateColor TextColor = FSlateColor::UseForeground();
+		if (OutlinerSettings->ColorByLevel)
+		{
+			TextColor = FGeometryCollectionTreeItem::GetColorPerDepth(Level);
+		}
+		FStatRow Row
+		{
+			FText::Format(LOCTEXT("Level", "Level {0}"), FText::AsNumber(Level)),
+			FText::AsNumber(Stats.CountsPerLevel[Level]),
+			TextColor
+		};
+		StatItems.Emplace(MakeShared<FStatRow>(MoveTemp(Row)));
+	}
+
+	FStatRow Row
+	{
+		FText(LOCTEXT("Embedded", "Embedded")),
+		FText::AsNumber(Stats.EmbeddedCount),
+		FSlateColor::UseForeground()
+	};
+	StatItems.Emplace(MakeShared<FStatRow>(MoveTemp(Row)));
+
+
+	StatListView->RequestListRefresh();
+}
+
+TSharedRef<ITableRow> SGeometryCollectionStatistics::HandleGenerateRow(FListItemPtr InItem, const TSharedRef<STableViewBase>& OwnerTable) const
+{	
+	return SNew(SGeometryCollectionStatisticsRow, OwnerTable)
+		.Name(InItem->Name)
+		.Value(InItem->Value)
+		.TextColor(InItem->TextColor);
+}
+
+#undef LOCTEXT_NAMESPACE /** SGeometryCollectionStatistics */
diff --git a/Engine/Plugins/Experimental/ChaosEditor/Source/FractureEditor/Private/SGeometryCollectionStatistics.h b/Engine/Plugins/Experimental/ChaosEditor/Source/FractureEditor/Private/SGeometryCollectionStatistics.h
new file mode 100644
--- /dev/null
+++ b/Engine/Plugins/Experimental/ChaosEditor/Source/FractureEditor/Private/SGeometryCollectionStatistics.h
@@ -0,0 +1,56 @@
+// Copyright Epic Games, Inc. All Rights Reserved.
+
+#pragma once
+
+#include "CoreMinimal.h"
+#include "Widgets/DeclarativeSyntaxSupport.h"
+#include "Widgets/SCompoundWidget.h"
+#include "Widgets/Views/SListView.h"
+
+struct FGeometryCollectionStatistics
+{
+	TArray<uint32> CountsPerLevel;
+	uint32 EmbeddedCount;
+};
+
+class SGeometryCollectionStatistics : public SCompoundWidget
+{
+public:
+	SLATE_BEGIN_ARGS(SGeometryCollectionStatistics)
+	{}
+	SLATE_END_ARGS()
+
+	/**
+	* Constructs a geom collection statistics widget.
+	*
+	* @param InArgs The Slate argument list.
+	*/
+	void Construct(const FArguments& InArgs);
+
+	void SetStatistics(const FGeometryCollectionStatistics& Stats);
+
+	void Refresh();
+
+private:
+
+	/** Holds basic information about a stat. */
+	struct FStatRow
+	{
+		FText Name;
+		FText Value;
+		FSlateColor TextColor;
+	};
+
+	using FListItemPtr = TSharedPtr<FStatRow>;
+
+	/** Callback for generating package info rows. */
+	TSharedRef<ITableRow> HandleGenerateRow(FListItemPtr InItem, const TSharedRef<STableViewBase>& OwnerTable) const;
+
+private:
+
+	/** Holds the ListView of package details rows. */
+	TSharedPtr<SListView<FListItemPtr>>	StatListView;
+
+	/** Holds the package details rows. */
+	TArray<FListItemPtr> StatItems;
+};
\ No newline at end of file
diff --git a/Engine/Plugins/Experimental/ChaosEditor/Source/FractureEditor/Public/FractureEditorModeToolkit.h b/Engine/Plugins/Experimental/ChaosEditor/Source/FractureEditor/Public/FractureEditorModeToolkit.h
--- a/Engine/Plugins/Experimental/ChaosEditor/Source/FractureEditor/Public/FractureEditorModeToolkit.h
+++ b/Engine/Plugins/Experimental/ChaosEditor/Source/FractureEditor/Public/FractureEditorModeToolkit.h
@@ -18,6 +18,8 @@ class FFractureToolContext;
 class FGeometryCollection;
 class SGeometryCollectionOutliner;
 class SGeometryCollectionHistogram;
+class SGeometryCollectionStatistics;
+struct FGeometryCollectionStatistics;
 class AGeometryCollectionActor;
 class UGeometryCollectionComponent;
 class UGeometryCollection;
@@ -73,6 +75,16 @@ private:
 	FFractureEditorModeToolkit* Toolkit;
 };
 
+struct FTextAndSlateColor
+{
+	FTextAndSlateColor(const FText& InText, const FSlateColor& InColor)
+		: Text(InText)
+		, Color(InColor)
+	{}
+	FText Text;
+	FSlateColor Color;
+};
+
 class FRACTUREEDITOR_API FFractureEditorModeToolkit : public FModeToolkit, public FGCObject
 {
 public:
@@ -143,7 +155,7 @@ public:
 	static void AddAdditionalAttributesIfRequired(UGeometryCollection* GeometryCollectionObject);
 	int32 GetLevelCount();
 
-	FText GetStatisticsSummary() const;
+	void GetStatisticsSummary(FGeometryCollectionStatistics& Stats) const;
 	FText GetSelectionInfo() const;
 
 	/** Returns the number of Mode specific tabs in the mode toolbar **/ 
@@ -218,5 +230,6 @@ private:
 	FMinorTabConfig HierarchyTabInfo;
 	TWeakPtr<SDockTab> StatisticsTab;
 	FMinorTabConfig StatisticsTabInfo;
+	TSharedPtr<SGeometryCollectionStatistics> StatisticsView;
 
 };
