Fracture editor - add multi column support to outliner
- displayed relative size as one of the column

diff --git a/Engine/Plugins/Experimental/ChaosEditor/Source/FractureEditor/Private/SGeometryCollectionOutliner.cpp b/Engine/Plugins/Experimental/ChaosEditor/Source/FractureEditor/Private/SGeometryCollectionOutliner.cpp
--- a/Engine/Plugins/Experimental/ChaosEditor/Source/FractureEditor/Private/SGeometryCollectionOutliner.cpp
+++ b/Engine/Plugins/Experimental/ChaosEditor/Source/FractureEditor/Private/SGeometryCollectionOutliner.cpp
@@ -6,6 +6,7 @@
 #include "Widgets/Text/STextBlock.h"
 #include "Widgets/Views/STreeView.h"
 #include "Widgets/Views/STableRow.h"
+#include "PropertyInfoViewStyle.h"
 
 #include "LevelEditor.h"
 #include "ToolMenus.h"
@@ -164,12 +165,76 @@ UOutlinerSettings::UOutlinerSettings(const FObjectInitializer& ObjInit)
 	, ColorByLevel(false)
 {}
 
+TSharedRef<SWidget> SGeometryCollectionOutlinerRow::GenerateWidgetForColumn(const FName& ColumnName)
+{
+	if (ColumnName == SGeometryCollectionOutlinerColumnID::Bone)
+	{
+		const TSharedPtr<SWidget> NameWidget = Item->MakeNameColumnWidget();
+		return SNew(SHorizontalBox)
+			+ SHorizontalBox::Slot()
+				.HAlign(HAlign_Left)
+				.VAlign(VAlign_Center)
+				.AutoWidth()
+				[
+					SNew(SExpanderArrow, SharedThis(this))
+					.ShouldDrawWires(true)
+				]
+			+ SHorizontalBox::Slot()
+				[
+					NameWidget.ToSharedRef()
+				];
+	}
+	if (ColumnName == SGeometryCollectionOutlinerColumnID::RelativeSize)
+	{
+		return Item->MakeRelativeSizeColumnWidget();
+	}
+	if (ColumnName == SGeometryCollectionOutlinerColumnID::DamageThreshold)
+	{
+		return Item->MakeDamageThresholdColumnWidget();
+	}
+	if (ColumnName == SGeometryCollectionOutlinerColumnID::InitialState)
+	{
+		return Item->MakeInitialStateColumnWidget();
+	}
+	return SNew(SHorizontalBox)
+	+ SHorizontalBox::Slot()
+		[
+			SNew(STextBlock)
+		];;
+}
 
 void SGeometryCollectionOutliner::Construct(const FArguments& InArgs)
 {
 	BoneSelectionChangedDelegate = InArgs._OnBoneSelectionChanged;
 	bPerformingSelection = false;
 
+	HeaderRowWidget =
+		SNew( SHeaderRow )
+			.Visibility( EVisibility::Visible );
+	HeaderRowWidget->AddColumn(
+		SHeaderRow::Column(SGeometryCollectionOutlinerColumnID::Bone)
+		.DefaultLabel(LOCTEXT("GCOutliner_Column_Bone", "Bone"))
+		.FillWidth(2.0f)
+	);
+	HeaderRowWidget->AddColumn(
+		SHeaderRow::Column(SGeometryCollectionOutlinerColumnID::RelativeSize)
+		.DefaultLabel(LOCTEXT("GCOutliner_Column_RelativeSize", "Relative Size"))
+		.HAlignHeader(EHorizontalAlignment::HAlign_Right)
+		.FillWidth(1.0f)
+	);
+	// HeaderRowWidget->AddColumn(
+	// 	SHeaderRow::Column(SGeometryCollectionOutlinerColumnID::DamageThreshold)
+	// 	.DefaultLabel(LOCTEXT("GCOutliner_Column_DamageThreshold", "Damage Threshold"))
+	// 	.HAlignHeader(EHorizontalAlignment::HAlign_Right)
+	// 	.FillWidth(1.0f)
+	// 	.Visibility(EVisibility::Collapsed)
+	// );
+	HeaderRowWidget->AddColumn(
+		SHeaderRow::Column(SGeometryCollectionOutlinerColumnID::InitialState)
+		.DefaultLabel(LOCTEXT("GCOutliner_Column_InitialState", "Initial State"))
+		.FillWidth(1.0f)
+	);
+	
 	ChildSlot
 	[
 		SAssignNew(TreeView, STreeView<FGeometryCollectionTreeItemPtr>)
@@ -183,6 +248,7 @@ void SGeometryCollectionOutliner::Construct(const FArguments& InArgs)
 		.OnGeneratePinnedRow(this, &SGeometryCollectionOutliner::OnGeneratePinnedRowWidget, true)
 		.HighlightParentNodesForSelection(true)
 		.OnSetExpansionRecursive(this, &SGeometryCollectionOutliner::ExpandRecursive)
+		.HeaderRow(HeaderRowWidget)
 	];
 }
 
@@ -468,6 +534,25 @@ void FGeometryCollectionTreeItemComponent::GetChildrenForBone(FGeometryCollectio
 	}
 }
 
+bool FGeometryCollectionTreeItemComponent::HasChildrenForBone(const FGeometryCollectionTreeItemBone& BoneItem) const
+{
+	if (!Component.IsValid())
+	{
+		return false;
+	}
+	if(const UGeometryCollection* RestCollection = Component->GetRestCollection())
+	{
+		if (const FGeometryCollection* Collection = RestCollection->GetGeometryCollection().Get())
+		{
+			if (const int32* BoneIndex = GuidIndexMap.Find(BoneItem.GetGuid()))
+			{
+				return Collection->Children[(*BoneIndex)].Num() > 0;
+			}
+		}
+	}
+	return false;
+}
+
 FText FGeometryCollectionTreeItemComponent::GetDisplayNameForBone(const FGuid& Guid) const
 {
 	if (!Component.IsValid())
@@ -597,126 +682,153 @@ bool FGeometryCollectionTreeItemComponent::FilterBoneIndex(int32 BoneIndex) cons
 	return true;	
 }
 
-TSharedRef<ITableRow> FGeometryCollectionTreeItemBone::MakeTreeRowWidget(const TSharedRef<STableViewBase>& InOwnerTable, bool bNoExtraColumn)
+void FGeometryCollectionTreeItemBone::UpdateItemFromCollection()
 {
+	const int32 ItemBoneIndex = GetBoneIndex();
+
+	constexpr FLinearColor InvalidColor(0.1f, 0.1f, 0.1f);
+	
+	ItemColor = InvalidColor;
+	RelativeSize = 0;
+	DamageThreshold = 0;
+	InitialState = INDEX_NONE;
+	
+	// Name / ItemText
 	UOutlinerSettings* OutlinerSettings = GetMutableDefault<UOutlinerSettings>();
-	FText ItemText = ParentComponentItem->GetDisplayNameForBone(Guid);
+	ItemText = ParentComponentItem->GetDisplayNameForBone(Guid);
 	if (OutlinerSettings->ItemText == EOutlinerItemNameEnum::BoneIndex)
 	{
-		ItemText = FText::FromString(FString::FromInt(GetBoneIndex()));
+		ItemText = FText::FromString(FString::FromInt(ItemBoneIndex));
 	}
 
 	// Set color according to simulation type
-
-	FSlateColor TextColor(FLinearColor::Red); // default color indicates something wrong
-	int32 InitialDynamicState = INDEX_NONE;
-
 	const UGeometryCollection* RestCollection = ParentComponentItem->GetComponent()->GetRestCollection();
 	if (RestCollection && IsValidChecked(RestCollection))
 	{
 		TSharedPtr<FGeometryCollection, ESPMode::ThreadSafe> GeometryCollectionPtr = RestCollection->GetGeometryCollection();
 		const TManagedArray<int32>& SimulationType = GeometryCollectionPtr->SimulationType;
-		if (ensure(GetBoneIndex() >= 0 && GetBoneIndex() < SimulationType.Num()))
+		const TManagedArray<float>* RelativeSizes = GeometryCollectionPtr->FindAttribute<float>("Size", FTransformCollection::TransformGroup);
+		
+		if (ensure(ItemBoneIndex >= 0 && ItemBoneIndex < SimulationType.Num()))
 		{
 			const bool bHasLevelAttribute = GeometryCollectionPtr->HasAttribute("Level", FGeometryCollection::TransformGroup);
 			if (bHasLevelAttribute && OutlinerSettings->ColorByLevel)
 			{
 				TManagedArray<int32>& Level = GeometryCollectionPtr->GetAttribute<int32>("Level", FTransformCollection::TransformGroup);
-				TextColor = FSlateColor(FGeometryCollectionTreeItem::GetColorPerDepth((uint32)Level[GetBoneIndex()]));
+				ItemColor = FSlateColor(FGeometryCollectionTreeItem::GetColorPerDepth((uint32)Level[ItemBoneIndex]));
 			}
 			else
 			{
-			switch (SimulationType[GetBoneIndex()])
-			{
-			case FGeometryCollection::ESimulationTypes::FST_None:
-				TextColor = FLinearColor::Green;
-				break;
-
-			case FGeometryCollection::ESimulationTypes::FST_Rigid:
-				if (GeometryCollectionPtr->IsVisible(GetBoneIndex()))
+				switch (SimulationType[ItemBoneIndex])
 				{
-					TextColor = FSlateColor::UseForeground();
-				}
-				else
-				{
-					TextColor = FLinearColor{ 0.1f, 0.1f, 0.1f, 1.f };
-				}
-				break;
+				case FGeometryCollection::ESimulationTypes::FST_None:
+					ItemColor = FLinearColor::Green;
+					break;
 
-			case FGeometryCollection::ESimulationTypes::FST_Clustered:
-				TextColor = FSlateColor(FColor::Cyan);
-				break;
+				case FGeometryCollection::ESimulationTypes::FST_Rigid:
+					if (GeometryCollectionPtr->IsVisible(ItemBoneIndex))
+					{
+						ItemColor = FSlateColor::UseForeground();
+					}
+					else
+					{
+						ItemColor = InvalidColor;
+					}
+					break;
 
-			default:
-				ensureMsgf(false, TEXT("Invalid Geometry Collection simulation type encountered."));
-				break;
-			}
+				case FGeometryCollection::ESimulationTypes::FST_Clustered:
+					ItemColor = FSlateColor(FColor::Cyan);
+					break;
+
+				default:
+					ensureMsgf(false, TEXT("Invalid Geometry Collection simulation type encountered."));
+					break;
+				}
 			}
 
-			InitialDynamicState = GeometryCollectionPtr->InitialDynamicState[GetBoneIndex()];
-		}
-		else
-		{
-			// UI contains invalid bone indices; likely not correctly updated after an undo or redo?
-			TextColor = FLinearColor(0.1f, 0.1f, 0.1f);
-			InitialDynamicState = INDEX_NONE;
+			InitialState = GeometryCollectionPtr->InitialDynamicState[ItemBoneIndex];
+			if (RelativeSizes)
+			{
+				RelativeSize = (*RelativeSizes)[ItemBoneIndex];
+			}
 		}
 	}
-	else
-	{
-		// Deleted rest collection
-		TextColor = FLinearColor(0.1f, 0.1f, 0.1f);
-		InitialDynamicState = INDEX_NONE;
-	}
-	
-	if (bNoExtraColumn)
+}
+
+TSharedRef<ITableRow> FGeometryCollectionTreeItemBone::MakeTreeRowWidget(const TSharedRef<STableViewBase>& InOwnerTable, bool bIsPinned)
+{
+	UpdateItemFromCollection();
+	if (bIsPinned)
 	{
-		// short version 
 		return SNew(STableRow<FGeometryCollectionTreeItemPtr>, InOwnerTable)
-		.Content()
-		[
-			SNew(SHorizontalBox)
-			+ SHorizontalBox::Slot()
-				.Padding(2.0f, 4.0f)
-				.AutoWidth()
-				[
-					SNew(STextBlock)
-					.Text(ItemText)
-					.ColorAndOpacity(TextColor)
-				]
-		]
-		.OnDragDetected(this, &FGeometryCollectionTreeItem::OnDragDetected)
-		.OnDrop(this, &FGeometryCollectionTreeItem::OnDrop)
-		.OnDragEnter(this, &FGeometryCollectionTreeItem::OnDragEnter)
-		.OnDragLeave(this, &FGeometryCollectionTreeItem::OnDragLeave);
+			.Content()
+			[
+				SNew(SHorizontalBox)
+				+ SHorizontalBox::Slot()
+					.Padding(2.0f, 4.0f)
+					.AutoWidth()
+					[
+						SNew(STextBlock)
+						.Text(ItemText)
+						.ColorAndOpacity(ItemColor)
+					]
+			];
 	}
-	// full item 
-	return SNew(STableRow<FGeometryCollectionTreeItemPtr>, InOwnerTable)
-		.Content()
-		[
-			SNew(SHorizontalBox)
+	return SNew(SGeometryCollectionOutlinerRow, InOwnerTable, SharedThis(this));
+}
+
+TSharedRef<SWidget> FGeometryCollectionTreeItemBone::MakeNameColumnWidget() const
+{
+	return SNew(SHorizontalBox)
 			+ SHorizontalBox::Slot()
-				.Padding(2.0f, 4.0f)
 				.AutoWidth()
 				[
 					SNew(STextBlock)
 					.Text(ItemText)
-					.ColorAndOpacity(TextColor)
-				]
-			+ SHorizontalBox::Slot()
-					.Padding(2.0f, 2.0f)
-					.FillWidth(1.0f)
-					.HAlign(HAlign_Right)
-					[
-						SNew(STextBlock)
-						.Text(GettextFromInitialDynamicState(InitialDynamicState))
-						.ColorAndOpacity(TextColor)
-					]
-		]
-		.OnDragDetected(this, &FGeometryCollectionTreeItem::OnDragDetected)
-		.OnDrop(this, &FGeometryCollectionTreeItem::OnDrop)
-		.OnDragEnter(this, &FGeometryCollectionTreeItem::OnDragEnter)
-		.OnDragLeave(this, &FGeometryCollectionTreeItem::OnDragLeave);
+					.ColorAndOpacity(ItemColor)
+				];
+}
+
+TSharedRef<SWidget> FGeometryCollectionTreeItemBone::MakeRelativeSizeColumnWidget() const
+{
+	static const FNumberFormattingOptions FormatOptions = FNumberFormattingOptions()
+				.SetMinimumFractionalDigits(2)
+				.SetMaximumFractionalDigits(2);
+	
+	return SNew(SHorizontalBox)
+		+ SHorizontalBox::Slot()
+			.Padding(12.f, 0.f)
+			.HAlign(HAlign_Right)
+			[
+				SNew(STextBlock)
+				.Text(FText::AsNumber(RelativeSize, &FormatOptions))
+				.ColorAndOpacity(ItemColor)
+			];
+}
+
+TSharedRef<SWidget> FGeometryCollectionTreeItemBone::MakeDamageThresholdColumnWidget() const
+{
+	return SNew(SHorizontalBox)
+		+ SHorizontalBox::Slot()
+			.Padding(12.f, 0.f)
+			.HAlign(HAlign_Right)
+			[
+				SNew(STextBlock)
+				.Text(FText::AsNumber(DamageThreshold))
+				.ColorAndOpacity(ItemColor)
+			];
+}
+
+TSharedRef<SWidget> FGeometryCollectionTreeItemBone::MakeInitialStateColumnWidget() const
+{
+	return SNew(SHorizontalBox)
+		+ SHorizontalBox::Slot()
+			.Padding(12.f, 0.f)
+			[
+				SNew(STextBlock)
+				.Text(GettextFromInitialDynamicState(InitialState))
+				.ColorAndOpacity(ItemColor)
+			];
 }
 
 void FGeometryCollectionTreeItemBone::GetChildren(FGeometryCollectionTreeItemList& OutChildren)
@@ -724,5 +836,10 @@ void FGeometryCollectionTreeItemBone::GetChildren(FGeometryCollectionTreeItemLis
 	ParentComponentItem->GetChildrenForBone(*this, OutChildren);
 }
 
+bool FGeometryCollectionTreeItemBone::HasChildren() const 
+{
+	return ParentComponentItem->HasChildrenForBone(*this);
+}
+
 
 #undef LOCTEXT_NAMESPACE // "ChaosEditor"
\ No newline at end of file
diff --git a/Engine/Plugins/Experimental/ChaosEditor/Source/FractureEditor/Private/SGeometryCollectionOutliner.h b/Engine/Plugins/Experimental/ChaosEditor/Source/FractureEditor/Private/SGeometryCollectionOutliner.h
--- a/Engine/Plugins/Experimental/ChaosEditor/Source/FractureEditor/Private/SGeometryCollectionOutliner.h
+++ b/Engine/Plugins/Experimental/ChaosEditor/Source/FractureEditor/Private/SGeometryCollectionOutliner.h
@@ -58,7 +58,7 @@ class FGeometryCollectionTreeItem : public TSharedFromThis<FGeometryCollectionTr
 {
 public:
 	virtual ~FGeometryCollectionTreeItem() {}
-	virtual TSharedRef<ITableRow> MakeTreeRowWidget(const TSharedRef<STableViewBase>& InOwnerTable, bool bNoExtraColumn = false) = 0;
+	virtual TSharedRef<ITableRow> MakeTreeRowWidget(const TSharedRef<STableViewBase>& InOwnerTable, bool bIsPinned = false) = 0;
 	virtual void GetChildren(FGeometryCollectionTreeItemList& OutChildren) = 0;
 	virtual UGeometryCollectionComponent* GetComponent() const = 0;
 
@@ -87,7 +87,7 @@ public:
 	}
 
 	/** FGeometryCollectionTreeItem interface */
-	virtual TSharedRef<ITableRow> MakeTreeRowWidget(const TSharedRef<STableViewBase>& InOwnerTable, bool bNoExtraColumn = false);
+	virtual TSharedRef<ITableRow> MakeTreeRowWidget(const TSharedRef<STableViewBase>& InOwnerTable, bool bIsPinned = false);
 	virtual void GetChildren(FGeometryCollectionTreeItemList& OutChildren) override;
 	virtual UGeometryCollectionComponent* GetComponent() const override { return Component.Get(); }
 
@@ -95,6 +95,7 @@ public:
 	FGeometryCollectionTreeItemPtr GetItemFromBoneIndex(int32 BoneIndex) const;
 	 
 	void GetChildrenForBone(FGeometryCollectionTreeItemBone& BoneItem, FGeometryCollectionTreeItemList& OutChildren) const;
+	bool HasChildrenForBone(const FGeometryCollectionTreeItemBone& BoneItem) const;
 	FText GetDisplayNameForBone(const FGuid& Guid) const;
 
 	void ExpandAll();
@@ -130,24 +131,76 @@ public:
 		: Guid(NewGuid)
 		, BoneIndex(InBoneIndex)
 		, ParentComponentItem(InParentComponentItem)
+		, ItemColor(FSlateColor::UseForeground())
+		, RelativeSize(0)
+		, DamageThreshold(0)
+		, InitialState(INDEX_NONE)
 	{}
 
 	/** FGeometryCollectionTreeItem interface */
-	TSharedRef<ITableRow> MakeTreeRowWidget(const TSharedRef<STableViewBase>& InOwnerTable, bool bNoExtraColumn = false);
+	TSharedRef<ITableRow> MakeTreeRowWidget(const TSharedRef<STableViewBase>& InOwnerTable, bool bIsPinned = false);
+	TSharedRef<SWidget> MakeNameColumnWidget() const;
+	TSharedRef<SWidget> MakeRelativeSizeColumnWidget() const;
+	TSharedRef<SWidget> MakeDamageThresholdColumnWidget() const;
+	TSharedRef<SWidget> MakeInitialStateColumnWidget() const;
 	virtual void GetChildren(FGeometryCollectionTreeItemList& OutChildren) override;
 	virtual int32 GetBoneIndex() const override { return BoneIndex; }
 	virtual UGeometryCollectionComponent* GetComponent() const { return ParentComponentItem->GetComponent(); }
 	const FGuid& GetGuid() const { return Guid; }
-
+	bool HasChildren() const;
+	
 protected:
 	virtual FReply OnDragDetected(const FGeometry& MyGeometry, const FPointerEvent& MouseEvent) override;
 	virtual FReply OnDrop(const FDragDropEvent& DragDropEvent) override;
 	virtual void OnDragEnter(FDragDropEvent const& InDragDropEvent) override;
 
+private:
+	void UpdateItemFromCollection();
+	
 private:
 	const FGuid Guid;
 	const int32 BoneIndex;
 	FGeometryCollectionTreeItemComponent* ParentComponentItem;
+	
+	FSlateColor ItemColor;
+	FText ItemText;
+	float RelativeSize;
+	float DamageThreshold;
+	int32 InitialState;
+};
+
+typedef TSharedPtr<class FGeometryCollectionTreeItemBone> FGeometryCollectionTreeItemBonePtr;
+
+namespace SGeometryCollectionOutlinerColumnID
+{
+	const FName Bone("Bone");
+	const FName RelativeSize("RelativeSize");
+	const FName DamageThreshold("DamageThreshold");
+	const FName InitialState("InitialState");
+}
+
+class SGeometryCollectionOutlinerRow : public SMultiColumnTableRow<FGeometryCollectionTreeItemBonePtr>
+{
+protected:
+	FGeometryCollectionTreeItemBonePtr Item;
+
+public:
+	SLATE_BEGIN_ARGS(SGeometryCollectionOutlinerRow) {}
+	SLATE_END_ARGS()
+
+	virtual TSharedRef<SWidget> GenerateWidgetForColumn(const FName& ColumnName) override;
+
+	void Construct(const FArguments& InArgs, TSharedRef<STableViewBase> OwnerTableView, FGeometryCollectionTreeItemBonePtr InItemToEdit)
+	{
+		Item = InItemToEdit;
+		SMultiColumnTableRow<FGeometryCollectionTreeItemBonePtr>::Construct(
+		FSuperRowType::FArguments()
+			.OnDragDetected(Item.Get(), &FGeometryCollectionTreeItem::OnDragDetected)
+			.OnDrop(Item.Get(), &FGeometryCollectionTreeItem::OnDrop)
+			.OnDragEnter(Item.Get(), &FGeometryCollectionTreeItem::OnDragEnter)
+			.OnDragLeave(Item.Get(), &FGeometryCollectionTreeItem::OnDragLeave)
+		, OwnerTableView);
+	}
 };
 
 class SGeometryCollectionOutliner: public SCompoundWidget
@@ -188,6 +241,7 @@ private:
 	void OnSelectionChanged(FGeometryCollectionTreeItemPtr Item, ESelectInfo::Type SelectInfo);
 private:
 	TSharedPtr<STreeView<FGeometryCollectionTreeItemPtr>> TreeView;
+	TSharedPtr< SHeaderRow > HeaderRowWidget;
 	TArray<TSharedPtr<FGeometryCollectionTreeItemComponent>> RootNodes;
 	FOnBoneSelectionChanged BoneSelectionChangedDelegate;
 	bool bPerformingSelection;
