add a merge-selected-only option to the fracture GeoMerge tool
+ stop selecting all after every GeoMerge

diff --git a/Engine/Plugins/Experimental/ChaosEditor/Source/FractureEditor/Private/FractureToolFixTinyGeo.cpp b/Engine/Plugins/Experimental/ChaosEditor/Source/FractureEditor/Private/FractureToolFixTinyGeo.cpp
--- a/Engine/Plugins/Experimental/ChaosEditor/Source/FractureEditor/Private/FractureToolFixTinyGeo.cpp
+++ b/Engine/Plugins/Experimental/ChaosEditor/Source/FractureEditor/Private/FractureToolFixTinyGeo.cpp
@@ -101,21 +101,29 @@ void UFractureToolFixTinyGeo::FractureContextChanged()
 		double MinVolume = GetMinVolume(Volumes);
 
 		TArray<int32> SmallIndices;
-		FindSmallBones(
-			Collection,
-			TArrayView<int32>(),
-			Volumes,
-			MinVolume,
-			SmallIndices
-		);
 
-		if (TinyGeoSettings->bAlsoMergeSelected)
+		if (TinyGeoSettings->UseBoneSelection == EUseBoneSelection::OnlyMergeSelected)
 		{
-			const TArray<int32> Selection = FractureContext.GetSelection();
-			for (int32 Bone : Selection)
+			SmallIndices = FractureContext.GetSelection();
+		}
+		else
+		{
+			FindSmallBones(
+				Collection,
+				TArrayView<int32>(),
+				Volumes,
+				MinVolume,
+				SmallIndices
+			);
+
+			if (TinyGeoSettings->UseBoneSelection == EUseBoneSelection::AlsoMergeSelected)
 			{
-				// note: we AddUnique when doing the merge, but for rendering we allow duplicates
-				SmallIndices.Add(Bone);
+				const TArray<int32>& Selection = FractureContext.GetSelection();
+				for (int32 Bone : Selection)
+				{
+					// note: we AddUnique when doing the merge, but for rendering we allow duplicates
+					SmallIndices.Add(Bone);
+				}
 			}
 		}
 
@@ -169,10 +177,18 @@ double UFractureToolFixTinyGeo::GetMinVolume(TArray<double>& Volumes)
 
 int32 UFractureToolFixTinyGeo::ExecuteFracture(const FFractureToolContext& FractureContext)
 {
-	int32 UpdateIdx = INDEX_NONE;
 	if (FractureContext.GetGeometryCollection().IsValid())
 	{
 		TArray<int32> TransformIndices; // left empty to refer to all indices
+		if (TinyGeoSettings->UseBoneSelection == EUseBoneSelection::OnlyMergeSelected)
+		{
+			// restrict to selection
+			TransformIndices = FractureContext.GetSelection();
+			if (TransformIndices.IsEmpty())
+			{
+				return INDEX_NONE;
+			}
+		}
 
 		TArray<double> Volumes;
 		FGeometryCollection& Collection = *FractureContext.GetGeometryCollection();
@@ -186,20 +202,27 @@ int32 UFractureToolFixTinyGeo::ExecuteFracture(const FFractureToolContext& Fract
 		double MinVolume = GetMinVolume(Volumes);
 
 		TArray<int32> SmallIndices;
-		FindSmallBones(
-			Collection,
-			TransformIndices,
-			Volumes,
-			MinVolume,
-			SmallIndices
-		);
-
-		if (TinyGeoSettings->bAlsoMergeSelected)
+		if (TinyGeoSettings->UseBoneSelection == EUseBoneSelection::OnlyMergeSelected)
+		{
+			SmallIndices = FractureContext.GetSelection();
+		}
+		else
 		{
-			const TArray<int32> Selection = FractureContext.GetSelection();
-			for (int32 Bone : Selection)
+			FindSmallBones(
+				Collection,
+				TransformIndices,
+				Volumes,
+				MinVolume,
+				SmallIndices
+			);
+
+			if (TinyGeoSettings->UseBoneSelection == EUseBoneSelection::AlsoMergeSelected)
 			{
-				SmallIndices.AddUnique(Bone);
+				const TArray<int32>& Selection = FractureContext.GetSelection();
+				for (int32 Bone : Selection)
+				{
+					SmallIndices.AddUnique(Bone);
+				}
 			}
 		}
 
@@ -219,10 +242,10 @@ int32 UFractureToolFixTinyGeo::ExecuteFracture(const FFractureToolContext& Fract
 			false, /*bUnionJoinedPieces*/ // Note: Union-ing the pieces is nicer in theory, but can leave cracks and non-manifold garbage
 			SelectionMethod
 		);
-		UpdateIdx = 0;
 	}
 
-	return UpdateIdx;
+	// Don't update selection
+	return INDEX_NONE;
 }
 
 
diff --git a/Engine/Plugins/Experimental/ChaosEditor/Source/FractureEditor/Private/FractureToolFixTinyGeo.h b/Engine/Plugins/Experimental/ChaosEditor/Source/FractureEditor/Private/FractureToolFixTinyGeo.h
--- a/Engine/Plugins/Experimental/ChaosEditor/Source/FractureEditor/Private/FractureToolFixTinyGeo.h
+++ b/Engine/Plugins/Experimental/ChaosEditor/Source/FractureEditor/Private/FractureToolFixTinyGeo.h
@@ -27,6 +27,14 @@ enum class ENeighborSelectionMethod
 	NearestCenter
 };
 
+UENUM()
+enum class EUseBoneSelection
+{
+	NoEffect,
+	AlsoMergeSelected,
+	OnlyMergeSelected
+};
+
 /** Settings controlling how geometry is selected and merged into neighboring geometry */
 UCLASS(config = EditorPerProjectUserSettings)
 class UFractureTinyGeoSettings : public UFractureToolSettings
@@ -42,9 +50,9 @@ public:
 	UPROPERTY(EditAnywhere, Category = MergeSettings)
 	ENeighborSelectionMethod NeighborSelection = ENeighborSelectionMethod::LargestNeighbor;
 
-	/** Also merge selected bones to their neighbors */
-	UPROPERTY(EditAnywhere, Category = MergeSettings, meta = (DisplayName = "Merge Selected"))
-	bool bAlsoMergeSelected = false;
+	/** Options for using the current bone selection */
+	UPROPERTY(EditAnywhere, Category = MergeSettings, meta = (DisplayName = "Bone Selection"))
+	EUseBoneSelection UseBoneSelection;
 
 
 	UPROPERTY(EditAnywhere, Category = FilterSettings)
diff --git a/Engine/Plugins/Experimental/ChaosEditor/Source/FractureEditor/Public/FractureTool.cpp b/Engine/Plugins/Experimental/ChaosEditor/Source/FractureEditor/Public/FractureTool.cpp
--- a/Engine/Plugins/Experimental/ChaosEditor/Source/FractureEditor/Public/FractureTool.cpp
+++ b/Engine/Plugins/Experimental/ChaosEditor/Source/FractureEditor/Public/FractureTool.cpp
@@ -308,6 +308,7 @@ void UFractureModalTool::Execute(TWeakPtr<FFractureEditorModeToolkit> InToolkit)
 		{
 			FGeometryCollectionEdit EditCollection(FractureContext.GetGeometryCollectionComponent(), GeometryCollection::EEditUpdate::RestPhysicsDynamic, !ExecuteUpdatesShape());
 
+			int32 InitialNumTransforms = FractureContext.GetGeometryCollection()->NumElements(FGeometryCollection::TransformGroup);
 			int32 FirstNewGeometryIndex = ExecuteFracture(FractureContext);
 			
 			if (FirstNewGeometryIndex > INDEX_NONE)
@@ -326,10 +327,21 @@ void UFractureModalTool::Execute(TWeakPtr<FFractureEditorModeToolkit> InToolkit)
 				}
 
 				FractureContext.SetSelection(NewTransforms);
-
-				Toolkit->RegenerateHistogram();
+			}
+			else
+			{
+				// either no update was done, or the updated range was not expressible as a final geometry index
+				// -- in the latter case, selection may become invalid, so we should clear it
+				int32 NumTransforms = FractureContext.GetGeometryCollection()->NumElements(FGeometryCollection::TransformGroup);
+				TArray<int32>& Selection = FractureContext.GetSelection();
+				if (InitialNumTransforms != NumTransforms)
+				{
+					Selection.Empty();
+				}
+				FractureContext.Sanitize(false);
 			}
 			
+			Toolkit->RegenerateHistogram();
 
 			FGeometryCollectionClusteringUtility::UpdateHierarchyLevelOfChildren(FractureContext.GetGeometryCollection().Get(), -1);
 
diff --git a/Engine/Plugins/Experimental/PlanarCutPlugin/Source/PlanarCut/Private/PlanarCut.cpp b/Engine/Plugins/Experimental/PlanarCutPlugin/Source/PlanarCut/Private/PlanarCut.cpp
--- a/Engine/Plugins/Experimental/PlanarCutPlugin/Source/PlanarCut/Private/PlanarCut.cpp
+++ b/Engine/Plugins/Experimental/PlanarCutPlugin/Source/PlanarCut/Private/PlanarCut.cpp
@@ -824,7 +824,7 @@ void FindBoneVolumes(
 		{
 			int32 TransformIdx = TransformIndicesArray[Idx];
 			int32 GeomIdx = Collection.TransformToGeometryIndex[TransformIdx];
-			if (GeomIdx == -1)
+			if (GeomIdx == -1 || !Collection.IsRigid(TransformIdx))
 			{
 				OutVolumes[Idx] = 0.0;
 			}
@@ -846,9 +846,9 @@ void FilterBonesByVolume(
 {
 	OutSmallBones.Reset();
 
-	auto AddIdx = [&Collection, &Volumes, &Filter, &OutSmallBones](int32 TransformIdx)
+	auto AddIdx = [&Collection, &Volumes, &Filter, &OutSmallBones](int32 TransformIdx, double Volume)
 	{
-		if (Collection.TransformToGeometryIndex[TransformIdx] > -1 && Filter(Volumes[TransformIdx], TransformIdx))
+		if (Collection.IsRigid(TransformIdx) && Collection.TransformToGeometryIndex[TransformIdx] > -1 && Filter(Volume, TransformIdx))
 		{
 			OutSmallBones.Add(TransformIdx);
 		}
@@ -864,7 +864,7 @@ void FilterBonesByVolume(
 		}
 		for (int32 TransformIdx = 0; TransformIdx < NumTransforms; TransformIdx++)
 		{
-			AddIdx(TransformIdx);
+			AddIdx(TransformIdx, Volumes[TransformIdx]);
 		}
 	}
 	else
@@ -873,9 +873,10 @@ void FilterBonesByVolume(
 		{
 			return;
 		}
-		for (int32 TransformIdx : TransformIndices)
+		for (int32 Idx = 0; Idx < TransformIndices.Num(); ++Idx)
 		{
-			AddIdx(TransformIdx);
+			int32 TransformIdx = TransformIndices[Idx];
+			AddIdx(TransformIdx, Volumes[Idx]);
 		}
 	}
 }
@@ -1041,22 +1042,22 @@ int32 MergeBones(
 		return Center;
 	};
 
-	for (int32 TransformIdx : TransformIndices)
+	for (int32 Idx = 0; Idx < TransformIndices.Num(); ++Idx)
 	{
+		int32 TransformIdx = TransformIndices[Idx];
 		int32 GeomIdx = Collection.TransformToGeometryIndex[TransformIdx];
 		if (GeomIdx > -1 && Collection.IsRigid(TransformIdx))
 		{
 			CanMerge.Add(GeomIdx);
-			GeomToVol.Add(GeomIdx, Volumes[TransformIdx]);
+			GeomToVol.Add(GeomIdx, Volumes[Idx]);
 		}
 	}
 	for (int32 TransformIdx : SmallTransformIndices)
 	{
 		int32 GeomIdx = Collection.TransformToGeometryIndex[TransformIdx];
-		if (GeomIdx > -1 && Collection.IsRigid(TransformIdx))
+		if (GeomIdx > -1 && Collection.IsRigid(TransformIdx) && GeomToVol.Contains(GeomIdx))
 		{
 			TooSmalls.Add(GeomIdx);
-			GeomToVol.Add(GeomIdx, Volumes[TransformIdx]);
 		}
 		else
 		{
