Fracture tool : fix crash when compiling a geometry collection based blueprint while having the fracture tool opened
- Make sure we check that existing selected components are registered when selection change

diff --git a/Engine/Plugins/Experimental/ChaosEditor/Source/FractureEditor/Private/FractureEditorMode.cpp b/Engine/Plugins/Experimental/ChaosEditor/Source/FractureEditor/Private/FractureEditorMode.cpp
--- a/Engine/Plugins/Experimental/ChaosEditor/Source/FractureEditor/Private/FractureEditorMode.cpp
+++ b/Engine/Plugins/Experimental/ChaosEditor/Source/FractureEditor/Private/FractureEditorMode.cpp
@@ -493,7 +493,7 @@ void UFractureEditorMode::OnActorSelectionChanged(const TArray<UObject*>& NewSel
 	// reset state for components no longer selected
 	for (UGeometryCollectionComponent* ExistingSelection : SelectedGeometryComponents)
 	{
-		if (ExistingSelection && !NewGeomSelection.Contains(ExistingSelection))
+		if (ExistingSelection && ExistingSelection->IsRegistered() && !NewGeomSelection.Contains(ExistingSelection))
 		{
 			// This component is no longer selected, clear any modified state
 
@@ -518,11 +518,7 @@ void UFractureEditorMode::OnActorSelectionChanged(const TArray<UObject*>& NewSel
 
 	SelectedGeometryComponents = NewGeomSelection.Array();
 
-	if(Toolkit.IsValid())
-	{
-		FFractureEditorModeToolkit* FractureToolkit = (FFractureEditorModeToolkit*)Toolkit.Get();
-		FractureToolkit->SetOutlinerComponents(SelectedGeometryComponents);
-	}
+	RefreshOutlinerWithCurrentSelection();
 }
 
 void UFractureEditorMode::RefreshOutlinerWithCurrentSelection()
