stop repeatedly reindexing materials on reset of geometry collection, and rely on the later call to GeometryCollectionObject->InitializeMaterials to handle it

diff --git a/Engine/Plugins/Experimental/ChaosEditor/Source/FractureEditor/Private/FractureToolGenerators.cpp b/Engine/Plugins/Experimental/ChaosEditor/Source/FractureEditor/Private/FractureToolGenerators.cpp
--- a/Engine/Plugins/Experimental/ChaosEditor/Source/FractureEditor/Private/FractureToolGenerators.cpp
+++ b/Engine/Plugins/Experimental/ChaosEditor/Source/FractureEditor/Private/FractureToolGenerators.cpp
@@ -413,16 +413,16 @@ void UFractureToolResetAsset::Execute(TWeakPtr<FFractureEditorModeToolkit> InToo
 					const UObject* SourceMesh = Source.SourceGeometryObject.TryLoad();
 					if (const UStaticMesh* SourceStaticMesh = Cast<UStaticMesh>(SourceMesh))
 					{
-						FGeometryCollectionConversion::AppendStaticMesh(SourceStaticMesh, Source.SourceMaterial, Source.LocalTransform, GeometryCollectionObject, true);
+						FGeometryCollectionConversion::AppendStaticMesh(SourceStaticMesh, Source.SourceMaterial, Source.LocalTransform, GeometryCollectionObject, false);
 					}
 					else if (const USkeletalMesh* SourceSkeletalMesh = Cast<USkeletalMesh>(SourceMesh))
 					{
 						// #todo (bmiller) Once we've settled on the right approach with static meshes, we'll need to apply the same strategy to skeletal mesh reconstruction.
-						// FGeometryCollectionConversion::AppendSkeletalMesh(SourceSkeletalMesh, Source.SourceMaterial, Source.LocalTransform, GeometryCollectionObject, true);
+						// FGeometryCollectionConversion::AppendSkeletalMesh(SourceSkeletalMesh, Source.SourceMaterial, Source.LocalTransform, GeometryCollectionObject, false);
 					}
 					else if (const UGeometryCollection* SourceGeometryCollection = Cast<UGeometryCollection>(SourceMesh))
 					{
-						FGeometryCollectionConversion::AppendGeometryCollection(SourceGeometryCollection, Source.SourceMaterial, Source.LocalTransform, GeometryCollectionObject, true);
+						FGeometryCollectionConversion::AppendGeometryCollection(SourceGeometryCollection, Source.SourceMaterial, Source.LocalTransform, GeometryCollectionObject, false);
 					}
 				}
 
