add "Keep Frac" setting to the fracture interactive selection tool, to allow dropping a random fraction of the selection filter

diff --git a/Engine/Plugins/Experimental/ChaosEditor/Source/FractureEditor/Private/FractureToolSelection.cpp b/Engine/Plugins/Experimental/ChaosEditor/Source/FractureEditor/Private/FractureToolSelection.cpp
--- a/Engine/Plugins/Experimental/ChaosEditor/Source/FractureEditor/Private/FractureToolSelection.cpp
+++ b/Engine/Plugins/Experimental/ChaosEditor/Source/FractureEditor/Private/FractureToolSelection.cpp
@@ -525,6 +525,13 @@ bool UFractureToolSelection::GetBonesByVolume(const FGeometryCollection& Collect
 			FilterIndices.Add(BoneIdx);
 		}
 	}
+
+	int TargetNumBones = FMath::Max(0, FilterIndices.Num() * SelectionSettings->KeepFrac);
+	FRandomStream Random(1); // Randomly remove but always with the same seed, to avoid flickering
+	while (FilterIndices.Num() > TargetNumBones)
+	{
+		FilterIndices.RemoveAtSwap(Random.RandHelper(FilterIndices.Num()), 1, false);
+	}
 	return true;
 }
 
diff --git a/Engine/Plugins/Experimental/ChaosEditor/Source/FractureEditor/Private/FractureToolSelection.h b/Engine/Plugins/Experimental/ChaosEditor/Source/FractureEditor/Private/FractureToolSelection.h
--- a/Engine/Plugins/Experimental/ChaosEditor/Source/FractureEditor/Private/FractureToolSelection.h
+++ b/Engine/Plugins/Experimental/ChaosEditor/Source/FractureEditor/Private/FractureToolSelection.h
@@ -82,6 +82,10 @@ public:
 	UPROPERTY(EditAnywhere, Category = FilterSettings, meta = (ClampMin = ".00001", UIMin = ".01", UIMax = "1", EditCondition = "VolumeSelectionMethod != EVolumeSelectionMethod::CubeRoot", EditConditionHides))
 	double MaxVolumeFrac = 0;
 
+	/** Fraction of bones to keep in the selection: If less than 1, bones will be randomly excluded from the selection filter */
+	UPROPERTY(EditAnywhere, Category = FilterSettings, meta = (ClampMin = ".00001", UIMin = ".01", UIMax = "1"))
+	double KeepFrac = 1;
+
 };
 
 
