make fracture interactive selection "keep fraction" option attempt to space out the bones it keeps, rather than choosing completely at random. Also give user control of the random seed.

diff --git a/Engine/Plugins/Experimental/ChaosEditor/Source/FractureEditor/Private/FractureToolSelection.cpp b/Engine/Plugins/Experimental/ChaosEditor/Source/FractureEditor/Private/FractureToolSelection.cpp
--- a/Engine/Plugins/Experimental/ChaosEditor/Source/FractureEditor/Private/FractureToolSelection.cpp
+++ b/Engine/Plugins/Experimental/ChaosEditor/Source/FractureEditor/Private/FractureToolSelection.cpp
@@ -19,6 +19,8 @@
 #include "Engine/Selection.h"
 #include "ProfilingDebugging/ScopedTimers.h"
 
+#include "Spatial/PriorityOrderPoints.h"
+
 #define LOCTEXT_NAMESPACE "FractureSelection"
 
 
@@ -375,6 +377,36 @@ TArray<UObject*> UFractureToolSelection::GetSettingsObjects() const
 	return Settings;
 }
 
+TArray<FBox> GetAllBounds(const FGeometryCollection& Collection, FTransform OuterTransform = FTransform::Identity)
+{
+	// Build bounding boxes for every transform (including clusters)
+	TArray<FTransform> AllTransforms;
+	TArray<FBox> AllBounds;
+	int32 NumBones = Collection.NumElements(FGeometryCollection::TransformGroup);
+	AllBounds.SetNum(NumBones);
+	GeometryCollectionAlgo::GlobalMatrices(Collection.Transform, Collection.Parent, AllTransforms);
+	FBox EmptyBounds(EForceInit::ForceInit);
+	AllBounds.Init(EmptyBounds, NumBones);
+	TArray<int32> BoneIndices = GeometryCollectionAlgo::ComputeRecursiveOrder(Collection);
+	for (int32 BoneIdx : BoneIndices)
+	{
+		int32 GeoIdx = Collection.TransformToGeometryIndex[BoneIdx];
+		if (GeoIdx != INDEX_NONE)
+		{
+			const FTransform& InnerTransform = AllTransforms[BoneIdx];
+			FTransform CombinedTransform = InnerTransform * OuterTransform;
+
+			AllBounds[BoneIdx] += Collection.BoundingBox[GeoIdx].TransformBy(CombinedTransform);
+		}
+		int32 ParentIdx = Collection.Parent[BoneIdx];
+		if (ParentIdx != INDEX_NONE)
+		{
+			AllBounds[ParentIdx] += AllBounds[BoneIdx];
+		}
+	}
+	return AllBounds;
+}
+
 
 void UFractureToolSelection::FractureContextChanged()
 {
@@ -414,30 +446,7 @@ void UFractureToolSelection::FractureContextChanged()
 		FTransform OuterTransform = FractureContext.GetTransform();
 
 		// Build bounding boxes for every transform (including clusters) in the global space of the component
-		TArray<FTransform> AllTransforms;
-		TArray<FBox> AllBounds;
-		int32 NumBones = Collection.NumElements(FGeometryCollection::TransformGroup);
-		AllBounds.SetNum(NumBones);
-		GeometryCollectionAlgo::GlobalMatrices(Collection.Transform, Collection.Parent, AllTransforms);
-		FBox EmptyBounds(EForceInit::ForceInit);
-		AllBounds.Init(EmptyBounds, NumBones);
-		TArray<int32> BoneIndices = GeometryCollectionAlgo::ComputeRecursiveOrder(Collection);
-		for (int32 BoneIdx : BoneIndices)
-		{
-			int32 GeoIdx = Collection.TransformToGeometryIndex[BoneIdx];
-			if (GeoIdx != INDEX_NONE)
-			{
-				const FTransform& InnerTransform = AllTransforms[BoneIdx];
-				FTransform CombinedTransform = InnerTransform * OuterTransform;
-
-				AllBounds[BoneIdx] += Collection.BoundingBox[GeoIdx].TransformBy(CombinedTransform);
-			}
-			int32 ParentIdx = Collection.Parent[BoneIdx];
-			if (ParentIdx != INDEX_NONE)
-			{
-				AllBounds[ParentIdx] += AllBounds[BoneIdx];
-			}
-		}
+		TArray<FBox> AllBounds = GetAllBounds(Collection, OuterTransform);
 
 		// Add bounding boxes for the selected indices to the visualization
 		for (int32 TransformIdx : SelectedIndices)
@@ -526,12 +535,34 @@ bool UFractureToolSelection::GetBonesByVolume(const FGeometryCollection& Collect
 		}
 	}
 
-	int TargetNumBones = FMath::Max(0, FilterIndices.Num() * SelectionSettings->KeepFrac);
-	FRandomStream Random(1); // Randomly remove but always with the same seed, to avoid flickering
-	while (FilterIndices.Num() > TargetNumBones)
+	int TargetNumBones = FMath::Max(0, FilterIndices.Num() * SelectionSettings->KeepFraction);
+	if (TargetNumBones < FilterIndices.Num())
 	{
-		FilterIndices.RemoveAtSwap(Random.RandHelper(FilterIndices.Num()), 1, false);
+		FRandomStream Random(SelectionSettings->RandomSeed);
+
+		TArray<FBox> AllBounds = GetAllBounds(Collection);
+		TArray<FVector> Points;
+		TArray<float> Weights;
+		Points.Reserve(FilterIndices.Num());
+		Weights.Reserve(FilterIndices.Num());
+		for (int Idx : FilterIndices)
+		{
+			FVector Point = AllBounds[Idx].GetCenter();
+			Points.Add(Point);
+			// set random importance weights, so the point ordering will be shuffled
+			Weights.Add((float)Random.RandHelper(FilterIndices.Num()) / FilterIndices.Num());
+		}
+		UE::Geometry::FPriorityOrderPoints Ordering;
+		Ordering.ComputeUniformSpaced(Points, Weights, TargetNumBones, 1);
+		TArray<int32> ToKeep;
+		ToKeep.Reserve(TargetNumBones);
+		for (int32 OrderIdx = 0; OrderIdx < TargetNumBones; OrderIdx++)
+		{
+			ToKeep.Add(FilterIndices[Ordering.Order[OrderIdx]]);
+		}
+		FilterIndices = ToKeep;
 	}
+
 	return true;
 }
 
diff --git a/Engine/Plugins/Experimental/ChaosEditor/Source/FractureEditor/Private/FractureToolSelection.h b/Engine/Plugins/Experimental/ChaosEditor/Source/FractureEditor/Private/FractureToolSelection.h
--- a/Engine/Plugins/Experimental/ChaosEditor/Source/FractureEditor/Private/FractureToolSelection.h
+++ b/Engine/Plugins/Experimental/ChaosEditor/Source/FractureEditor/Private/FractureToolSelection.h
@@ -84,8 +84,11 @@ public:
 
 	/** Fraction of bones to keep in the selection: If less than 1, bones will be randomly excluded from the selection filter */
 	UPROPERTY(EditAnywhere, Category = FilterSettings, meta = (ClampMin = ".00001", UIMin = ".01", UIMax = "1"))
-	double KeepFrac = 1;
+	double KeepFraction = 1;
 
+	/** Seed to use for randomization when deciding which bones to keep w/ the Keep Fraction*/
+	UPROPERTY(EditAnywhere, Category = FilterSettings)
+	int RandomSeed = 1;
 };
 
 
