Make voronoi fracture run in a background thread w/ a regularly-sampled progress/cancel UI

diff --git a/Engine/Plugins/Experimental/ChaosEditor/Source/FractureEditor/FractureEditor.Build.cs b/Engine/Plugins/Experimental/ChaosEditor/Source/FractureEditor/FractureEditor.Build.cs
--- a/Engine/Plugins/Experimental/ChaosEditor/Source/FractureEditor/FractureEditor.Build.cs
+++ b/Engine/Plugins/Experimental/ChaosEditor/Source/FractureEditor/FractureEditor.Build.cs
@@ -51,6 +51,7 @@ public class FractureEditor : ModuleRules
                 "GeometryCollectionEngine",
                 "GeometryCollectionEditor",
 				"ModelingComponents",
+				"ModelingOperators",
 				"GeometryCore",
 				"MeshDescription",
 				"StaticMeshDescription",
diff --git a/Engine/Plugins/Experimental/ChaosEditor/Source/FractureEditor/Private/FractureToolBackgroundTask.h b/Engine/Plugins/Experimental/ChaosEditor/Source/FractureEditor/Private/FractureToolBackgroundTask.h
new file mode 100644
--- /dev/null
+++ b/Engine/Plugins/Experimental/ChaosEditor/Source/FractureEditor/Private/FractureToolBackgroundTask.h
@@ -0,0 +1,96 @@
+// Copyright Epic Games, Inc. All Rights Reserved.
+
+#pragma once
+
+#include "GeometryCollection/GeometryCollection.h"
+#include "ModelingTaskTypes.h"
+#include "ModelingOperators.h"
+#include "BackgroundModelingComputeSource.h"
+#include "Misc/ScopedSlowTask.h"
+
+namespace UE
+{
+namespace Fracture
+{
+
+// Base class for background operators that update geometry collections (e.g. to fracture in a background thread)
+class FGeometryCollectionOperator : public UE::Geometry::TGenericDataOperator<FGeometryCollection>
+{
+public:
+	int ResultGeometryIndex = -1;
+	TUniquePtr<FGeometryCollection> CollectionCopy;
+
+	FGeometryCollectionOperator(const FGeometryCollection& SourceCollection)
+	{
+		CollectionCopy = MakeUnique<FGeometryCollection>();
+		CollectionCopy->CopyMatchingAttributesFrom(SourceCollection, nullptr);
+	}
+
+	virtual ~FGeometryCollectionOperator() = default;
+
+	virtual int GetResultGeometryIndex()
+	{
+		return ResultGeometryIndex;
+	}
+
+	// Post-process the geometry collection on success -- by default, clears proximity data
+	virtual void OnSuccess(FGeometryCollection& Collection)
+	{
+		// Invalidate proximity
+		if (Collection.HasAttribute("Proximity", FGeometryCollection::GeometryGroup))
+		{
+			Collection.RemoveAttribute("Proximity", FGeometryCollection::GeometryGroup);
+		}
+	}
+};
+
+// Run a blocking geometry collection op, but with a responsive cancel option
+template<class GeometryCollectionOpType>
+int RunCancellableGeometryCollectionOp(FGeometryCollection& ToUpdate, TUniquePtr<GeometryCollectionOpType>&& NewOp, FText DefaultMessage, float DialogDelay = .5)
+{
+	FScopedSlowTask SlowTask(1, DefaultMessage);
+	SlowTask.MakeDialogDelayed(DialogDelay, true);
+	using FGeometryCollectionTask = UE::Geometry::TModelingOpTask<GeometryCollectionOpType>;
+	using FExecuter = UE::Geometry::FAsyncTaskExecuterWithAbort<FGeometryCollectionTask>;
+	TUniquePtr<FExecuter> BackgroundTask = MakeUnique<FExecuter>(MoveTemp(NewOp));
+	BackgroundTask->StartBackgroundTask();
+
+	bool bSuccess = false;
+	while (true)
+	{
+		if (BackgroundTask->IsDone())
+		{
+			bSuccess = !BackgroundTask->GetTask().IsAborted();
+			break;
+		}
+		if (SlowTask.ShouldCancel())
+		{
+			// Release ownership to the TDeleterTask that is spawned by CancelAndDelete()
+			BackgroundTask.Release()->CancelAndDelete();
+			break;
+		}
+		FPlatformProcess::Sleep(.2); // SlowTask::ShouldCancel will throttle any updates faster than .2 seconds
+		SlowTask.TickProgress(); // Note: EnterProgressFrame will also call this, but we aren't getting progress updates here yet
+		// TODO: Connect up a way to get progress updates from the task back to SlowTask
+	}
+
+	if (bSuccess)
+	{
+		check(BackgroundTask != nullptr && BackgroundTask->IsDone());
+		TUniquePtr<GeometryCollectionOpType> Op = BackgroundTask->GetTask().ExtractOperator();
+		TUniquePtr<FGeometryCollection> Result = Op->ExtractResult();
+		if (!Result.IsValid())
+		{
+			return -1;
+		}
+
+		ToUpdate.CopyMatchingAttributesFrom(*Result, nullptr);
+		Op->OnSuccess(ToUpdate);
+		return Op->GetResultGeometryIndex();
+	}
+
+	return -1;
+}
+
+} // end namespace UE::Fracture
+} // end namespace UE
\ No newline at end of file
diff --git a/Engine/Plugins/Experimental/ChaosEditor/Source/FractureEditor/Private/FractureToolCutter.cpp b/Engine/Plugins/Experimental/ChaosEditor/Source/FractureEditor/Private/FractureToolCutter.cpp
--- a/Engine/Plugins/Experimental/ChaosEditor/Source/FractureEditor/Private/FractureToolCutter.cpp
+++ b/Engine/Plugins/Experimental/ChaosEditor/Source/FractureEditor/Private/FractureToolCutter.cpp
@@ -18,6 +18,10 @@
 #include "Voronoi/Voronoi.h"
 #include "PlanarCut.h"
 
+#include "FractureToolBackgroundTask.h"
+
+using namespace UE::Fracture;
+
 #define LOCTEXT_NAMESPACE "FractureToolCutter"
 
 
@@ -349,6 +353,45 @@ void UFractureToolVoronoiCutterBase::FractureContextChanged()
 	UpdateVisualizations(FractureContexts);
 }
 
+
+class FVoronoiFractureOp : public FGeometryCollectionOperator
+{
+public:
+	FVoronoiFractureOp(const FGeometryCollection& SourceCollection) : FGeometryCollectionOperator(SourceCollection)
+	{}
+
+	virtual ~FVoronoiFractureOp() = default;
+
+	TArray<int> Selection;
+	FBox Bounds;
+	TArray<FVector> Sites;
+	TOptional<FNoiseSettings> NoiseSettings;
+	float PointSpacing;
+	float Grout;
+	int Seed;
+	FTransform Transform;
+
+	// TGenericDataOperator interface:
+	virtual void CalculateResult(FProgressCancel* Progress) override
+	{
+		FVoronoiDiagram Voronoi(Sites, Bounds, .1f);
+
+		FPlanarCells VoronoiPlanarCells = FPlanarCells(Sites, Voronoi);
+		VoronoiPlanarCells.InternalSurfaceMaterials.NoiseSettings = NoiseSettings;
+
+		if (Progress && Progress->Cancelled())
+		{
+			return;
+		}
+
+		// TODO: Pass the ProgressCancel down so the computation can be cancelled
+		ResultGeometryIndex = CutMultipleWithPlanarCells(VoronoiPlanarCells, *CollectionCopy, Selection, Grout, PointSpacing, Seed, Transform);
+		
+		SetResult(MoveTemp(CollectionCopy));
+	}
+};
+
+
 int32 UFractureToolVoronoiCutterBase::ExecuteFracture(const FFractureToolContext& FractureContext)
 {
 	if (FractureContext.IsValid())
@@ -356,22 +399,25 @@ int32 UFractureToolVoronoiCutterBase::ExecuteFracture(const FFractureToolContext
 		TArray<FVector> Sites;
 		GenerateVoronoiSites(FractureContext, Sites);
 		FBox VoronoiBounds = GetVoronoiBounds(FractureContext, Sites);
-			
-		FVoronoiDiagram Voronoi(Sites, VoronoiBounds, .1f);
-
-		FPlanarCells VoronoiPlanarCells = FPlanarCells(Sites, Voronoi);
 
-		FNoiseSettings NoiseSettings;
+		TUniquePtr<FVoronoiFractureOp> VoronoiOp = MakeUnique<FVoronoiFractureOp>(*(FractureContext.GetGeometryCollection()));
+		VoronoiOp->Bounds = VoronoiBounds;
+		VoronoiOp->Selection = FractureContext.GetSelection();
+		VoronoiOp->Grout = CutterSettings->Grout;
+		VoronoiOp->PointSpacing = CollisionSettings->GetPointSpacing();
+		VoronoiOp->Sites = Sites;
 		if (CutterSettings->Amplitude > 0.0f)
 		{
-			CutterSettings->TransferNoiseSettings(NoiseSettings);
-			VoronoiPlanarCells.InternalSurfaceMaterials.NoiseSettings = NoiseSettings;
+			FNoiseSettings Settings;
+			CutterSettings->TransferNoiseSettings(Settings);
+			VoronoiOp->NoiseSettings = Settings;
 		}
+		VoronoiOp->Seed = FractureContext.GetSeed();
+		VoronoiOp->Transform = FractureContext.GetTransform();
 
-		// Proximity is invalidated.
-		ClearProximity(FractureContext.GetGeometryCollection().Get());
-
-		return CutMultipleWithPlanarCells(VoronoiPlanarCells, *(FractureContext.GetGeometryCollection()), FractureContext.GetSelection(), CutterSettings->Grout, CollisionSettings->GetPointSpacing(), FractureContext.GetSeed(), FractureContext.GetTransform());
+		int Result = RunCancellableGeometryCollectionOp<FVoronoiFractureOp>(*(FractureContext.GetGeometryCollection()),
+			MoveTemp(VoronoiOp), LOCTEXT("ComputingVoronoiFractureMessage", "Computing Voronoi Fracture"));
+		return Result;
 	}
 
 	return INDEX_NONE;
diff --git a/Engine/Plugins/Runtime/MeshModelingToolset/Source/ModelingOperators/Public/ModelingTaskTypes.h b/Engine/Plugins/Runtime/MeshModelingToolset/Source/ModelingOperators/Public/ModelingTaskTypes.h
--- a/Engine/Plugins/Runtime/MeshModelingToolset/Source/ModelingOperators/Public/ModelingTaskTypes.h
+++ b/Engine/Plugins/Runtime/MeshModelingToolset/Source/ModelingOperators/Public/ModelingTaskTypes.h
@@ -86,7 +86,7 @@ public:
 	/**
 	 * Tells the child FAbandonableTask to terminate itself, via the bAbort flag
 	 * passed in SetAbortSource, and then starts a TDeleterTask that waits for
-	 * completion and 
+	 * completion to delete this task.
 	 */
 	void CancelAndDelete()
 	{
@@ -130,7 +130,7 @@ public:
 	}
 
 	/** @return true if the task should stop computation */
-	bool IsAborted()
+	bool IsAborted() const
 	{
 		return (bExternalAbortFlag != nullptr) ? *bExternalAbortFlag : false;
 	}
@@ -143,6 +143,7 @@ public:
 };
 
 
+
 } // end namespace UE::Geometry
 } // end namespace UE
 
