Fracture tool : make escape to deselect all if all tools have already been deselected

diff --git a/Engine/Plugins/Experimental/ChaosEditor/Source/FractureEditor/Private/FractureEditorModeToolkit.cpp b/Engine/Plugins/Experimental/ChaosEditor/Source/FractureEditor/Private/FractureEditorModeToolkit.cpp
--- a/Engine/Plugins/Experimental/ChaosEditor/Source/FractureEditor/Private/FractureEditorModeToolkit.cpp
+++ b/Engine/Plugins/Experimental/ChaosEditor/Source/FractureEditor/Private/FractureEditorModeToolkit.cpp
@@ -877,7 +877,17 @@ void FFractureEditorModeToolkit::BindCommands()
 
 	ToolkitCommands->MapAction(
 		Commands.CancelTool,
-		FExecuteAction::CreateLambda([=]() { this->SetActiveTool(nullptr); })
+		FExecuteAction::CreateLambda([=]()
+		{
+			if (GetActiveTool())
+			{
+				this->SetActiveTool(nullptr);
+			}
+			else
+			{
+				GEditor->SelectNone(true, true, false);
+			}
+		})
 	);
 
 	// Map actions of all the Fracture Tools
