Fracture Tool: prevent crash when compiling geometry collection based blueprint while in the fracture editor

diff --git a/Engine/Plugins/Experimental/ChaosEditor/Source/FractureEditor/Private/FractureEditorMode.cpp b/Engine/Plugins/Experimental/ChaosEditor/Source/FractureEditor/Private/FractureEditorMode.cpp
--- a/Engine/Plugins/Experimental/ChaosEditor/Source/FractureEditor/Private/FractureEditorMode.cpp
+++ b/Engine/Plugins/Experimental/ChaosEditor/Source/FractureEditor/Private/FractureEditorMode.cpp
@@ -493,7 +493,7 @@ void UFractureEditorMode::OnActorSelectionChanged(const TArray<UObject*>& NewSel
 	// reset state for components no longer selected
 	for (UGeometryCollectionComponent* ExistingSelection : SelectedGeometryComponents)
 	{
-		if (ExistingSelection && ExistingSelection->IsRegistered() && !NewGeomSelection.Contains(ExistingSelection))
+		if (ExistingSelection && ExistingSelection->IsRegistered() && !ExistingSelection->IsBeingDestroyed() && !NewGeomSelection.Contains(ExistingSelection))
 		{
 			// This component is no longer selected, clear any modified state
 
