Chaos destruction : add a damage watcher to collect collision impulse from PIE and display them in the outliner
- Add also a drop down for the type of column to display  ( InitialState, RelativeSize, Damages for now )

diff --git a/Engine/Plugins/Experimental/ChaosEditor/Source/FractureEditor/Private/FractureEditorModeToolkit.cpp b/Engine/Plugins/Experimental/ChaosEditor/Source/FractureEditor/Private/FractureEditorModeToolkit.cpp
--- a/Engine/Plugins/Experimental/ChaosEditor/Source/FractureEditor/Private/FractureEditorModeToolkit.cpp
+++ b/Engine/Plugins/Experimental/ChaosEditor/Source/FractureEditor/Private/FractureEditorModeToolkit.cpp
@@ -690,12 +690,18 @@ void FFractureEditorModeToolkit::OnObjectPostEditChange( UObject* Object, FPrope
 		}
 		else if (PropertyChangedEvent.Property->GetFName() == GET_MEMBER_NAME_CHECKED(UOutlinerSettings, ItemText))
 		{
-			UOutlinerSettings* OutlinerSettings = GetMutableDefault<UOutlinerSettings>();
 			OutlinerView->RegenerateItems();
 		}
 		else if (PropertyChangedEvent.Property->GetFName() == GET_MEMBER_NAME_CHECKED(UOutlinerSettings, ColorByLevel))
 		{
-			UOutlinerSettings* OutlinerSettings = GetMutableDefault<UOutlinerSettings>();
+			OutlinerView->RegenerateItems();
+			FGeometryCollectionStatistics Stats;
+			GetStatisticsSummary(Stats);
+			StatisticsView->SetStatistics(Stats);
+		}
+		else if (PropertyChangedEvent.Property->GetFName() == GET_MEMBER_NAME_CHECKED(UOutlinerSettings, DisplayColumn))
+		{
+			OutlinerView->UpdateDisplayColumn();
 			OutlinerView->RegenerateItems();
 			FGeometryCollectionStatistics Stats;
 			GetStatisticsSummary(Stats);
diff --git a/Engine/Plugins/Experimental/ChaosEditor/Source/FractureEditor/Private/SGeometryCollectionOutliner.cpp b/Engine/Plugins/Experimental/ChaosEditor/Source/FractureEditor/Private/SGeometryCollectionOutliner.cpp
--- a/Engine/Plugins/Experimental/ChaosEditor/Source/FractureEditor/Private/SGeometryCollectionOutliner.cpp
+++ b/Engine/Plugins/Experimental/ChaosEditor/Source/FractureEditor/Private/SGeometryCollectionOutliner.cpp
@@ -17,10 +17,10 @@
 #include "FractureSettings.h"
 #include "GeometryCollectionOutlinerDragDrop.h"
 #include "GeometryCollection/GeometryCollectionUtility.h"
+#include "PhysicsProxy/GeometryCollectionPhysicsProxy.h"
 
 #define LOCTEXT_NAMESPACE "ChaosEditor"
 
-
 static FText GettextFromInitialDynamicState(int32 InitialDynamicState)
 {
 	switch (InitialDynamicState)
@@ -132,7 +132,7 @@ FReply FGeometryCollectionTreeItemBone::OnDragDetected(const FGeometry& MyGeomet
 FReply FGeometryCollectionTreeItemBone::OnDrop(const FDragDropEvent& DragDropEvent)
 {
 	TSharedPtr<FDragDropOperation> Operation = DragDropEvent.GetOperation();
-	if (!Operation.IsValid())
+	if (!Operation.IsValid()) 
 	{
 		return FReply::Unhandled();
 	}
@@ -163,6 +163,7 @@ UOutlinerSettings::UOutlinerSettings(const FObjectInitializer& ObjInit)
 	: Super(ObjInit)
 	, ItemText(EOutlinerItemNameEnum::BoneIndex)
 	, ColorByLevel(false)
+	, DisplayColumn(EOutlinerDisplayColumn::InitialState)
 {}
 
 TSharedRef<SWidget> SGeometryCollectionOutlinerRow::GenerateWidgetForColumn(const FName& ColumnName)
@@ -184,17 +185,20 @@ TSharedRef<SWidget> SGeometryCollectionOutlinerRow::GenerateWidgetForColumn(cons
 					NameWidget.ToSharedRef()
 				];
 	}
-	if (ColumnName == SGeometryCollectionOutlinerColumnID::RelativeSize)
-	{
-		return Item->MakeRelativeSizeColumnWidget();
-	}
-	if (ColumnName == SGeometryCollectionOutlinerColumnID::DamageThreshold)
+	if (ColumnName == SGeometryCollectionOutlinerColumnID::Custom)
 	{
-		return Item->MakeDamageThresholdColumnWidget();
-	}
-	if (ColumnName == SGeometryCollectionOutlinerColumnID::InitialState)
-	{
-		return Item->MakeInitialStateColumnWidget();
+		UOutlinerSettings* OutlinerSettings = GetMutableDefault<UOutlinerSettings>();
+		switch (OutlinerSettings->DisplayColumn)
+		{
+		case EOutlinerDisplayColumn::InitialState:
+			return Item->MakeInitialStateColumnWidget();
+		case EOutlinerDisplayColumn::RelativeSize:
+			return Item->MakeRelativeSizeColumnWidget();
+		case EOutlinerDisplayColumn::Damages:
+			return Item->MakeDamagesColumnWidget();
+		default:
+			return Item->MakeEmptyColumnWidget();
+		}
 	}
 	return SNew(SHorizontalBox)
 	+ SHorizontalBox::Slot()
@@ -216,24 +220,13 @@ void SGeometryCollectionOutliner::Construct(const FArguments& InArgs)
 		.DefaultLabel(LOCTEXT("GCOutliner_Column_Bone", "Bone"))
 		.FillWidth(2.0f)
 	);
-	HeaderRowWidget->AddColumn(
-		SHeaderRow::Column(SGeometryCollectionOutlinerColumnID::RelativeSize)
-		.DefaultLabel(LOCTEXT("GCOutliner_Column_RelativeSize", "Relative Size"))
-		.HAlignHeader(EHorizontalAlignment::HAlign_Right)
-		.FillWidth(1.0f)
-	);
 	// HeaderRowWidget->AddColumn(
-	// 	SHeaderRow::Column(SGeometryCollectionOutlinerColumnID::DamageThreshold)
-	// 	.DefaultLabel(LOCTEXT("GCOutliner_Column_DamageThreshold", "Damage Threshold"))
-	// 	.HAlignHeader(EHorizontalAlignment::HAlign_Right)
-	// 	.FillWidth(1.0f)
-	// 	.Visibility(EVisibility::Collapsed)
+	//  	SHeaderRow::Column(SGeometryCollectionOutlinerColumnID::Custom)
+	//  	.DefaultLabel(LOCTEXT("GCOutliner_Column_Custom", "-"))
+	//  	.HAlignHeader(EHorizontalAlignment::HAlign_Right)
+	//  	.FillWidth(1.0f)
 	// );
-	HeaderRowWidget->AddColumn(
-		SHeaderRow::Column(SGeometryCollectionOutlinerColumnID::InitialState)
-		.DefaultLabel(LOCTEXT("GCOutliner_Column_InitialState", "Initial State"))
-		.FillWidth(1.0f)
-	);
+	UpdateDisplayColumn();
 	
 	ChildSlot
 	[
@@ -252,6 +245,54 @@ void SGeometryCollectionOutliner::Construct(const FArguments& InArgs)
 	];
 }
 
+void SGeometryCollectionOutliner::UpdateDisplayColumn()
+{
+	constexpr int32 CustomColumnIndex = 1;
+	constexpr int32 CustomFillWidth = 2.0f;
+	// first remove custom columns
+	HeaderRowWidget->RemoveColumn(SGeometryCollectionOutlinerColumnID::Custom);
+
+	// then add the right customn one based on the selection
+	UOutlinerSettings* OutlinerSettings = GetMutableDefault<UOutlinerSettings>();
+	switch (OutlinerSettings->DisplayColumn)
+	{
+	case EOutlinerDisplayColumn::InitialState:
+		HeaderRowWidget->InsertColumn(
+			SHeaderRow::Column(SGeometryCollectionOutlinerColumnID::Custom)
+				.DefaultLabel(LOCTEXT("GCOutliner_Column_InitialState", "Initial State"))
+				.FillWidth(CustomFillWidth),
+			CustomColumnIndex
+		);
+		break;
+	case EOutlinerDisplayColumn::RelativeSize:
+		HeaderRowWidget->InsertColumn(
+			SHeaderRow::Column(SGeometryCollectionOutlinerColumnID::Custom)
+				.DefaultLabel(LOCTEXT("GCOutliner_Column_RelativeSize", "Relative Size"))
+				.HAlignHeader(EHorizontalAlignment::HAlign_Right)
+				.FillWidth(CustomFillWidth),
+			CustomColumnIndex
+		);
+		break;
+	case EOutlinerDisplayColumn::Damages:
+		HeaderRowWidget->InsertColumn(
+			SHeaderRow::Column(SGeometryCollectionOutlinerColumnID::Custom)
+				.DefaultLabel(LOCTEXT("GCOutliner_Column_Damages", "Damages"))
+	 			.HAlignHeader(EHorizontalAlignment::HAlign_Center)
+	 			.FillWidth(CustomFillWidth),
+			CustomColumnIndex
+		);
+		break;
+	default:
+		HeaderRowWidget->InsertColumn(
+			SHeaderRow::Column(SGeometryCollectionOutlinerColumnID::Custom)
+				.DefaultLabel(LOCTEXT("GCOutliner_Column_Custom", "-"))
+				.HAlignHeader(EHorizontalAlignment::HAlign_Right)
+				.FillWidth(CustomFillWidth),
+			CustomColumnIndex
+		);
+	}
+}
+
 void SGeometryCollectionOutliner::RegenerateItems()
 {
 	TreeView->RebuildList();
@@ -702,7 +743,9 @@ void FGeometryCollectionTreeItemBone::UpdateItemFromCollection()
 	}
 
 	// Set color according to simulation type
-	const UGeometryCollection* RestCollection = ParentComponentItem->GetComponent()->GetRestCollection();
+	const UGeometryCollectionComponent* Component = ParentComponentItem->GetComponent();
+	const FCollisionImpulseCollector* Collector = Component->GetRunTimeDataCollector();
+	const UGeometryCollection* RestCollection = Component->GetRestCollection();
 	if (RestCollection && IsValidChecked(RestCollection))
 	{
 		TSharedPtr<FGeometryCollection, ESPMode::ThreadSafe> GeometryCollectionPtr = RestCollection->GetGeometryCollection();
@@ -751,6 +794,8 @@ void FGeometryCollectionTreeItemBone::UpdateItemFromCollection()
 			{
 				RelativeSize = (*RelativeSizes)[ItemBoneIndex];
 			}
+
+			Damages = Collector? (*Collector)[ItemBoneIndex].Max : 0;  
 		}
 	}
 }
@@ -806,6 +851,23 @@ TSharedRef<SWidget> FGeometryCollectionTreeItemBone::MakeRelativeSizeColumnWidge
 			];
 }
 
+TSharedRef<SWidget> FGeometryCollectionTreeItemBone::MakeDamagesColumnWidget() const
+{
+	static const FNumberFormattingOptions FormatOptions = FNumberFormattingOptions()
+			.SetMinimumFractionalDigits(1)
+			.SetMaximumFractionalDigits(1);
+	
+	return SNew(SHorizontalBox)
+		+ SHorizontalBox::Slot()
+			.Padding(12.f, 0.f)
+			.HAlign(HAlign_Right)
+			[
+				SNew(STextBlock)
+				.Text(FText::AsNumber(Damages, &FormatOptions))
+				.ColorAndOpacity(ItemColor)
+			];
+}
+
 TSharedRef<SWidget> FGeometryCollectionTreeItemBone::MakeDamageThresholdColumnWidget() const
 {
 	return SNew(SHorizontalBox)
@@ -831,6 +893,17 @@ TSharedRef<SWidget> FGeometryCollectionTreeItemBone::MakeInitialStateColumnWidge
 			];
 }
 
+TSharedRef<SWidget> FGeometryCollectionTreeItemBone::MakeEmptyColumnWidget() const
+{
+	return SNew(SHorizontalBox)
+		+ SHorizontalBox::Slot()
+			.Padding(12.f, 0.f)
+			[
+				SNew(STextBlock)
+				.ColorAndOpacity(ItemColor)
+			];
+}
+
 void FGeometryCollectionTreeItemBone::GetChildren(FGeometryCollectionTreeItemList& OutChildren)
 {
 	ParentComponentItem->GetChildrenForBone(*this, OutChildren);
diff --git a/Engine/Plugins/Experimental/ChaosEditor/Source/FractureEditor/Private/SGeometryCollectionOutliner.h b/Engine/Plugins/Experimental/ChaosEditor/Source/FractureEditor/Private/SGeometryCollectionOutliner.h
--- a/Engine/Plugins/Experimental/ChaosEditor/Source/FractureEditor/Private/SGeometryCollectionOutliner.h
+++ b/Engine/Plugins/Experimental/ChaosEditor/Source/FractureEditor/Private/SGeometryCollectionOutliner.h
@@ -35,6 +35,14 @@ enum class EOutlinerItemNameEnum : uint8
 	BoneIndex = 1					UMETA(DisplayName = "Bone Index"),
 };
 
+UENUM(BlueprintType)
+enum class EOutlinerDisplayColumn : uint8
+{
+	InitialState = 0		UMETA(DisplayName = "Initial State"),
+	RelativeSize = 1		UMETA(DisplayName = "Relative Size"),
+	Damages = 2				UMETA(Damages = "Damages"),
+};
+
 /** Settings for Outliner configuration. **/
 UCLASS()
 class UOutlinerSettings : public UObject
@@ -48,9 +56,13 @@ public:
 	UPROPERTY(EditAnywhere, Category = OutlinerSettings, meta = (DisplayName = "Item Text"))
 	EOutlinerItemNameEnum ItemText;
 
-	/** whther to use level coloring */
+	/** whether to use level coloring */
 	UPROPERTY(EditAnywhere, Category = OutlinerSettings, meta = (DisplayName = "Color By Level"))
 	bool ColorByLevel;
+
+	/** the column to be display in the outliner */
+	UPROPERTY(EditAnywhere, Category = OutlinerSettings, meta = (DisplayName = "Display column"))
+	EOutlinerDisplayColumn DisplayColumn;
 };
 
 
@@ -133,6 +145,7 @@ public:
 		, ParentComponentItem(InParentComponentItem)
 		, ItemColor(FSlateColor::UseForeground())
 		, RelativeSize(0)
+		, Damages(0)
 		, DamageThreshold(0)
 		, InitialState(INDEX_NONE)
 	{}
@@ -141,8 +154,10 @@ public:
 	TSharedRef<ITableRow> MakeTreeRowWidget(const TSharedRef<STableViewBase>& InOwnerTable, bool bIsPinned = false);
 	TSharedRef<SWidget> MakeNameColumnWidget() const;
 	TSharedRef<SWidget> MakeRelativeSizeColumnWidget() const;
+	TSharedRef<SWidget> MakeDamagesColumnWidget() const;
 	TSharedRef<SWidget> MakeDamageThresholdColumnWidget() const;
 	TSharedRef<SWidget> MakeInitialStateColumnWidget() const;
+	TSharedRef<SWidget> MakeEmptyColumnWidget() const;
 	virtual void GetChildren(FGeometryCollectionTreeItemList& OutChildren) override;
 	virtual int32 GetBoneIndex() const override { return BoneIndex; }
 	virtual UGeometryCollectionComponent* GetComponent() const { return ParentComponentItem->GetComponent(); }
@@ -165,6 +180,7 @@ private:
 	FSlateColor ItemColor;
 	FText ItemText;
 	float RelativeSize;
+	float Damages;
 	float DamageThreshold;
 	int32 InitialState;
 };
@@ -174,9 +190,7 @@ typedef TSharedPtr<class FGeometryCollectionTreeItemBone> FGeometryCollectionTre
 namespace SGeometryCollectionOutlinerColumnID
 {
 	const FName Bone("Bone");
-	const FName RelativeSize("RelativeSize");
-	const FName DamageThreshold("DamageThreshold");
-	const FName InitialState("InitialState");
+	const FName Custom("Custom");
 }
 
 class SGeometryCollectionOutlinerRow : public SMultiColumnTableRow<FGeometryCollectionTreeItemBonePtr>
@@ -218,6 +232,7 @@ public:
 	void Construct(const FArguments& InArgs);
 
 	void RegenerateItems();
+	void UpdateDisplayColumn();
 
 	TSharedRef<ITableRow> MakeTreeRowWidget(FGeometryCollectionTreeItemPtr InInfo, const TSharedRef<STableViewBase>& OwnerTable);
 	TSharedRef<ITableRow> OnGeneratePinnedRowWidget(FGeometryCollectionTreeItemPtr InItem, const TSharedRef<STableViewBase>& InOwnerTable, bool bPinned);
@@ -239,6 +254,7 @@ public:
 
 private:
 	void OnSelectionChanged(FGeometryCollectionTreeItemPtr Item, ESelectInfo::Type SelectInfo);
+	
 private:
 	TSharedPtr<STreeView<FGeometryCollectionTreeItemPtr>> TreeView;
 	TSharedPtr< SHeaderRow > HeaderRowWidget;
diff --git a/Engine/Source/Runtime/Experimental/Chaos/Private/PhysicsProxy/GeometryCollectionPhysicsProxy.cpp b/Engine/Source/Runtime/Experimental/Chaos/Private/PhysicsProxy/GeometryCollectionPhysicsProxy.cpp
--- a/Engine/Source/Runtime/Experimental/Chaos/Private/PhysicsProxy/GeometryCollectionPhysicsProxy.cpp
+++ b/Engine/Source/Runtime/Experimental/Chaos/Private/PhysicsProxy/GeometryCollectionPhysicsProxy.cpp
@@ -358,6 +358,7 @@ FGeometryCollectionPhysicsProxy::FGeometryCollectionPhysicsProxy(
 	const FSimulationParameters& SimulationParameters,
 	FCollisionFilterData InSimFilter,
 	FCollisionFilterData InQueryFilter,
+	FGuid InCollectorGuid,
 	const Chaos::EMultiBufferMode BufferMode)
 	: Base(InOwner)
 	, Parameters(SimulationParameters)
@@ -377,6 +378,7 @@ FGeometryCollectionPhysicsProxy::FGeometryCollectionPhysicsProxy(
 	, GameThreadCollection(GameThreadCollectionIn)
 	, bIsPhysicsThreadWorldTransformDirty(false)
 	, bIsCollisionFilterDataDirty(false)
+	, CollectorGuid(InCollectorGuid)
 {
 	// We rely on a guarded buffer.
 	check(BufferMode == Chaos::EMultiBufferMode::TripleGuarded);
@@ -2103,6 +2105,10 @@ void FGeometryCollectionPhysicsProxy::BufferPhysicsResults_Internal(Chaos::FPBDR
 			}    // end if
 
 			PhysicsThreadCollection.Active[TransformGroupIndex] = !TargetResults.States[TransformGroupIndex].DisabledState;
+
+#if WITH_EDITORONLY_DATA
+			TargetResults.CollisionImpulses[TransformGroupIndex] = static_cast<float>(Handle->CollisionImpulse());
+#endif			
 		}    // end for
 	}        // STAT_CalcParticleToWorld scope
 	
@@ -2176,10 +2182,20 @@ bool FGeometryCollectionPhysicsProxy::PullFromPhysicsState(const Chaos::FDirtyGe
 				{
 					Parenttype = (TargetResults.States[TransformGroupIndex].DynamicInternalClusterParent != 0)? Chaos::EInternalClusterType::Dynamic: Chaos::EInternalClusterType::KinematicOrStatic; 
 				}
-				(*InternalClusterParentTypeArray)[TransformGroupIndex] =static_cast<uint8>(Parenttype); 
+				(*InternalClusterParentTypeArray)[TransformGroupIndex] =static_cast<uint8>(Parenttype);
 			}
 		}
 
+#if WITH_EDITORONLY_DATA
+		if (FCollisionImpulseCollector* Collector = FCollisionImpulseWatcher::GetInstance().Find(CollectorGuid))
+		{
+			for (int32 TransformGroupIndex = 0; TransformGroupIndex < NumTransforms; ++TransformGroupIndex)
+			{
+				Collector->SampleImpulse(TransformGroupIndex, TargetResults.CollisionImpulses[TransformGroupIndex]);
+			}
+		}
+#endif
+		
 		// second : interpolate-able ones
 		const bool bNeedInterpolation = (NextPullData!= nullptr);
 		if (bNeedInterpolation)
@@ -3480,3 +3496,49 @@ void FGeometryCollectionPhysicsProxy::FieldForcesUpdateCallback(Chaos::FPBDRigid
 		}
 	}
 }
+
+
+void FCollisionImpulseCollector::Reset(int32 NumTransforms)
+{
+	ImpulseData.Reset();
+	ImpulseData.SetNum(NumTransforms);
+}
+
+void FCollisionImpulseCollector::SampleImpulse(int32 TransformIndex, float Value)
+{
+	if (ImpulseData.IsValidIndex(TransformIndex))
+	{
+		FImpulseData& Data = ImpulseData[TransformIndex];
+		Data.Min = FMath::Min(Data.Min, Value);
+		Data.Max = FMath::Max(Data.Max, Value);
+		Data.Count++;
+	}
+}
+
+FCollisionImpulseWatcher& FCollisionImpulseWatcher::GetInstance()
+{
+	static FCollisionImpulseWatcher Instance;
+	return Instance;
+}
+
+void FCollisionImpulseWatcher::Clear()
+{
+	Collectors.Reset();
+}
+
+FCollisionImpulseCollector* FCollisionImpulseWatcher::Find(const FGuid& Guid)
+{
+	return Collectors.Find(Guid);
+}
+
+void FCollisionImpulseWatcher::AddCollector(const FGuid& Guid, int32 TransformNum)
+{
+	FCollisionImpulseCollector& Collector = Collectors.FindOrAdd(Guid);
+	Collector.Reset(TransformNum);
+}
+
+void FCollisionImpulseWatcher::RemoveCollector(const FGuid& Guid)
+{
+	Collectors.Remove(Guid);
+	
+}
diff --git a/Engine/Source/Runtime/Experimental/Chaos/Public/GeometryCollectionProxyData.h b/Engine/Source/Runtime/Experimental/Chaos/Public/GeometryCollectionProxyData.h
--- a/Engine/Source/Runtime/Experimental/Chaos/Public/GeometryCollectionProxyData.h
+++ b/Engine/Source/Runtime/Experimental/Chaos/Public/GeometryCollectionProxyData.h
@@ -161,6 +161,9 @@ public:
 		
 		Transforms.SetNumUninitialized(NumTransforms);
 		Parent.SetNumUninitialized(NumTransforms);
+#if WITH_EDITORONLY_DATA
+		CollisionImpulses.SetNumUninitialized(NumTransforms);
+#endif		
 	}
 
 	struct FState
@@ -182,6 +185,11 @@ public:
 
 	TArray<FTransform> Transforms;
 	TArray<int32> Parent;
+
+#if WITH_EDITORONLY_DATA
+	// use to display impulse statistics in editor
+	TArray<float> CollisionImpulses;
+#endif
 	
 	bool IsObjectDynamic;
 	bool IsObjectLoading;
diff --git a/Engine/Source/Runtime/Experimental/Chaos/Public/PhysicsProxy/GeometryCollectionPhysicsProxy.h b/Engine/Source/Runtime/Experimental/Chaos/Public/PhysicsProxy/GeometryCollectionPhysicsProxy.h
--- a/Engine/Source/Runtime/Experimental/Chaos/Public/PhysicsProxy/GeometryCollectionPhysicsProxy.h
+++ b/Engine/Source/Runtime/Experimental/Chaos/Public/PhysicsProxy/GeometryCollectionPhysicsProxy.h
@@ -79,6 +79,7 @@ public:
 		const FSimulationParameters& SimulationParameters,
 		FCollisionFilterData InSimFilter,
 		FCollisionFilterData InQueryFilter,
+		FGuid InCollectorGuid,
 		const Chaos::EMultiBufferMode BufferMode=Chaos::EMultiBufferMode::TripleGuarded);
 	virtual ~FGeometryCollectionPhysicsProxy();
 
@@ -377,7 +378,49 @@ private:
 	// paradigm, at least for this component of the handshake.
 	Chaos::FGuardedTripleBuffer<FGeometryCollectionResults> PhysToGameInterchange;
 
-	FProxyInterpolationData InterpolationData; 
+	FProxyInterpolationData InterpolationData;
+
+	// this is used as a unique ID when collecting data from runtime
+	FGuid CollectorGuid;
+};
+
+/**
+ * the collision impulse watcher collects runtime data about
+ * collision impulses on each piece of the geometry collection 
+ */
+struct CHAOS_API FCollisionImpulseCollector
+{
+public:
+	struct FImpulseData
+	{
+		float Min = TNumericLimits<float>::Max();
+		float Max = 0;
+		int32 Count = 0;
+	};
+
+	void Reset(int32 NumTransforms);
+	int32 Num() const { return ImpulseData.Num(); }
+	const FImpulseData& operator[](int32 TransformIndex) const { return ImpulseData[TransformIndex]; }
+	void SampleImpulse(int32 TransformIndex, float Value);
+	
+private:
+	TArray<FImpulseData> ImpulseData;
+};
+
+struct CHAOS_API FCollisionImpulseWatcher
+{
+public:
+	void Clear();
+	void AddCollector(const FGuid& Guid, int32 TransformNum);
+	void RemoveCollector(const FGuid& Guid);
+
+	FCollisionImpulseCollector* Find(const FGuid& Guid);
+
+	static FCollisionImpulseWatcher& GetInstance();
+	
+private:
+	// collectors by geometry collection Guids
+	TMap<FGuid,FCollisionImpulseCollector> Collectors;
 };
 
 CHAOS_API TUniquePtr<Chaos::FTriangleMesh> CreateTriangleMesh(const int32 FaceStart,const int32 FaceCount,const TManagedArray<bool>& Visible,const TManagedArray<FIntVector>& Indices, bool bRotateWinding = true);
diff --git a/Engine/Source/Runtime/Experimental/GeometryCollectionEngine/Private/GeometryCollection/GeometryCollectionComponent.cpp b/Engine/Source/Runtime/Experimental/GeometryCollectionEngine/Private/GeometryCollection/GeometryCollectionComponent.cpp
--- a/Engine/Source/Runtime/Experimental/GeometryCollectionEngine/Private/GeometryCollection/GeometryCollectionComponent.cpp
+++ b/Engine/Source/Runtime/Experimental/GeometryCollectionEngine/Private/GeometryCollection/GeometryCollectionComponent.cpp
@@ -242,6 +242,10 @@ UGeometryCollectionComponent::UGeometryCollectionComponent(const FObjectInitiali
 	, bNotifyRemovals(false)
 	, bStoreVelocities(false)
 	, bShowBoneColors(false)
+#if WITH_EDITORONLY_DATA 
+	, bEnableRunTimeDataCollection(false)
+	, RunTimeDataCollectionGuid(FGuid::NewGuid())
+#endif
 	, bEnableReplication(false)
 	, bEnableAbandonAfterLevel(true)
 	, ReplicationAbandonClusterLevel(0)
@@ -2219,7 +2223,19 @@ void UGeometryCollectionComponent::RegisterAndInitializePhysicsProxy()
 		GetInitializationCommands(SimulationParameters.InitializationCommands);
 	}
 
-	PhysicsProxy = new FGeometryCollectionPhysicsProxy(this, *DynamicCollection, SimulationParameters, InitialSimFilter, InitialQueryFilter);
+	FGuid CollectorGuid = FGuid::NewGuid();
+#if WITH_EDITORONLY_DATA
+	CollectorGuid = RunTimeDataCollectionGuid;
+	if (bEnableRunTimeDataCollection && RestCollection)
+	{
+		FCollisionImpulseWatcher::GetInstance().AddCollector(CollectorGuid, RestCollection->NumElements(FGeometryCollection::TransformGroup));
+	}
+	else
+	{
+		FCollisionImpulseWatcher::GetInstance().RemoveCollector(CollectorGuid);
+	}
+#endif
+	PhysicsProxy = new FGeometryCollectionPhysicsProxy(this, *DynamicCollection, SimulationParameters, InitialSimFilter, InitialQueryFilter, CollectorGuid);
 	FPhysScene_Chaos* Scene = GetInnerChaosScene();
 	Scene->AddObject(this, PhysicsProxy);
 
@@ -2227,6 +2243,13 @@ void UGeometryCollectionComponent::RegisterAndInitializePhysicsProxy()
 	SetAsyncPhysicsTickEnabled(GetIsReplicated());
 }
 
+#if WITH_EDITORONLY_DATA
+const FCollisionImpulseCollector* UGeometryCollectionComponent::GetRunTimeDataCollector() const
+{
+	return FCollisionImpulseWatcher::GetInstance().Find(RunTimeDataCollectionGuid);
+}
+#endif
+
 void UGeometryCollectionComponent::OnDestroyPhysicsState()
 {
 	UActorComponent::OnDestroyPhysicsState();
diff --git a/Engine/Source/Runtime/Experimental/GeometryCollectionEngine/Public/GeometryCollection/GeometryCollectionComponent.h b/Engine/Source/Runtime/Experimental/GeometryCollectionEngine/Public/GeometryCollection/GeometryCollectionComponent.h
--- a/Engine/Source/Runtime/Experimental/GeometryCollectionEngine/Public/GeometryCollection/GeometryCollectionComponent.h
+++ b/Engine/Source/Runtime/Experimental/GeometryCollectionEngine/Public/GeometryCollection/GeometryCollectionComponent.h
@@ -35,6 +35,7 @@ class UChaosPhysicalMaterial;
 class AChaosSolverActor;
 struct FGeometryCollectionEmbeddedExemplar;
 class UInstancedStaticMeshComponent;
+struct FCollisionImpulseCollector;
 
 DECLARE_DYNAMIC_MULTICAST_DELEGATE_OneParam(FOnChaosBreakEvent, const FChaosBreakEvent&, BreakEvent);
 
@@ -708,6 +709,9 @@ public:
 	int32 EmbeddedIndexToTransformIndex(const UInstancedStaticMeshComponent* ISMComponent, int32 InstanceIndex) const;
 #endif
 
+#if WITH_EDITORONLY_DATA
+	const FCollisionImpulseCollector* GetRunTimeDataCollector() const;
+#endif
 	
 	// #todo should this only be available in editor?
 	void SetRestState(TArray<FTransform>&& InRestTransforms);
@@ -761,6 +765,14 @@ protected:
 	UPROPERTY(EditAnywhere, BlueprintReadWrite, Category = "ChaosPhysics|General")
 	bool bShowBoneColors;
 
+#if WITH_EDITORONLY_DATA
+	UPROPERTY(EditAnywhere, Category = "ChaosPhysics|General")
+	bool bEnableRunTimeDataCollection;
+	
+	UPROPERTY(BlueprintReadOnly, Category = "ChaosPhysics|General")
+	FGuid RunTimeDataCollectionGuid;
+#endif
+	
 	/** Populate the static geometry structures for the render thread. */
 	void InitConstantData(FGeometryCollectionConstantData* ConstantData) const;
 
@@ -895,5 +907,4 @@ private:
 
 	/** True if GeometryCollection transforms have changed from previous tick. */
 	bool bIsMoving;
-
 };
