Fracture Editor :
- Change display column to column mode
- Rename runtime data collectors classes
- Fix issue where collected damage values were invalid
- Added damage threshold and breaking state to runtime data
- Add refresh button to the outliner

diff --git a/Engine/Plugins/Experimental/ChaosEditor/Source/FractureEditor/Private/FractureEditorModeToolkit.cpp b/Engine/Plugins/Experimental/ChaosEditor/Source/FractureEditor/Private/FractureEditorModeToolkit.cpp
--- a/Engine/Plugins/Experimental/ChaosEditor/Source/FractureEditor/Private/FractureEditorModeToolkit.cpp
+++ b/Engine/Plugins/Experimental/ChaosEditor/Source/FractureEditor/Private/FractureEditorModeToolkit.cpp
@@ -604,13 +604,18 @@ TSharedRef<SDockTab> FFractureEditorModeToolkit::CreateHierarchyTab(const FSpawn
 					[
 						OutlinerDetailsView.ToSharedRef()
 					]
-
 					+ SVerticalBox::Slot()
 					[
 						SAssignNew(OutlinerView, SGeometryCollectionOutliner)
 						.OnBoneSelectionChanged(this, &FFractureEditorModeToolkit::OnOutlinerBoneSelectionChanged)
 					]
-
+					+ SVerticalBox::Slot()
+					.AutoHeight()
+					[
+						SNew(SButton)
+						.OnClicked(this, &FFractureEditorModeToolkit::RefreshOutliner)
+						.Text(LOCTEXT("Refresh", "Refresh"))
+					]
 					+ SVerticalBox::Slot()
 					.AutoHeight()
 					[
@@ -630,6 +635,15 @@ TSharedRef<SDockTab> FFractureEditorModeToolkit::CreateHierarchyTab(const FSpawn
 	return CreatedTab.ToSharedRef();
 }
 
+FReply FFractureEditorModeToolkit::RefreshOutliner()
+{
+	if (OutlinerView)
+	{
+		OutlinerView->RegenerateItems();
+	}
+	return FReply::Handled();
+}
+
 TSharedRef<SDockTab> FFractureEditorModeToolkit::CreateStatisticsTab(const FSpawnTabArgs& Args)
 {
 	FMargin MorePadding = FMargin(10.0f, 2.0f);
@@ -699,9 +713,9 @@ void FFractureEditorModeToolkit::OnObjectPostEditChange( UObject* Object, FPrope
 			GetStatisticsSummary(Stats);
 			StatisticsView->SetStatistics(Stats);
 		}
-		else if (PropertyChangedEvent.Property->GetFName() == GET_MEMBER_NAME_CHECKED(UOutlinerSettings, DisplayColumn))
+		else if (PropertyChangedEvent.Property->GetFName() == GET_MEMBER_NAME_CHECKED(UOutlinerSettings, ColumnMode))
 		{
-			OutlinerView->UpdateDisplayColumn();
+			OutlinerView->RegenerateHeader();
 			OutlinerView->RegenerateItems();
 			FGeometryCollectionStatistics Stats;
 			GetStatisticsSummary(Stats);
diff --git a/Engine/Plugins/Experimental/ChaosEditor/Source/FractureEditor/Private/SGeometryCollectionOutliner.cpp b/Engine/Plugins/Experimental/ChaosEditor/Source/FractureEditor/Private/SGeometryCollectionOutliner.cpp
--- a/Engine/Plugins/Experimental/ChaosEditor/Source/FractureEditor/Private/SGeometryCollectionOutliner.cpp
+++ b/Engine/Plugins/Experimental/ChaosEditor/Source/FractureEditor/Private/SGeometryCollectionOutliner.cpp
@@ -163,7 +163,7 @@ UOutlinerSettings::UOutlinerSettings(const FObjectInitializer& ObjInit)
 	: Super(ObjInit)
 	, ItemText(EOutlinerItemNameEnum::BoneIndex)
 	, ColorByLevel(false)
-	, DisplayColumn(EOutlinerDisplayColumn::InitialState)
+	, ColumnMode(EOutlinerColumnMode::StateAndSize)
 {}
 
 TSharedRef<SWidget> SGeometryCollectionOutlinerRow::GenerateWidgetForColumn(const FName& ColumnName)
@@ -185,26 +185,17 @@ TSharedRef<SWidget> SGeometryCollectionOutlinerRow::GenerateWidgetForColumn(cons
 					NameWidget.ToSharedRef()
 				];
 	}
-	if (ColumnName == SGeometryCollectionOutlinerColumnID::Custom)
-	{
-		UOutlinerSettings* OutlinerSettings = GetMutableDefault<UOutlinerSettings>();
-		switch (OutlinerSettings->DisplayColumn)
-		{
-		case EOutlinerDisplayColumn::InitialState:
-			return Item->MakeInitialStateColumnWidget();
-		case EOutlinerDisplayColumn::RelativeSize:
-			return Item->MakeRelativeSizeColumnWidget();
-		case EOutlinerDisplayColumn::Damages:
-			return Item->MakeDamagesColumnWidget();
-		default:
-			return Item->MakeEmptyColumnWidget();
-		}
-	}
-	return SNew(SHorizontalBox)
-	+ SHorizontalBox::Slot()
-		[
-			SNew(STextBlock)
-		];;
+	if (ColumnName == SGeometryCollectionOutlinerColumnID::InitialState)
+		return Item->MakeInitialStateColumnWidget();
+	if (ColumnName == SGeometryCollectionOutlinerColumnID::RelativeSize)
+		return Item->MakeRelativeSizeColumnWidget();
+	if (ColumnName == SGeometryCollectionOutlinerColumnID::Damage)
+		return Item->MakeDamagesColumnWidget();
+	if (ColumnName == SGeometryCollectionOutlinerColumnID::DamageThreshold)
+		return Item->MakeDamageThresholdColumnWidget();
+	if (ColumnName == SGeometryCollectionOutlinerColumnID::Broken)
+		return Item->MakeBrokenColumnWidget();
+	return Item->MakeEmptyColumnWidget();
 }
 
 void SGeometryCollectionOutliner::Construct(const FArguments& InArgs)
@@ -215,18 +206,7 @@ void SGeometryCollectionOutliner::Construct(const FArguments& InArgs)
 	HeaderRowWidget =
 		SNew( SHeaderRow )
 			.Visibility( EVisibility::Visible );
-	HeaderRowWidget->AddColumn(
-		SHeaderRow::Column(SGeometryCollectionOutlinerColumnID::Bone)
-		.DefaultLabel(LOCTEXT("GCOutliner_Column_Bone", "Bone"))
-		.FillWidth(2.0f)
-	);
-	// HeaderRowWidget->AddColumn(
-	//  	SHeaderRow::Column(SGeometryCollectionOutlinerColumnID::Custom)
-	//  	.DefaultLabel(LOCTEXT("GCOutliner_Column_Custom", "-"))
-	//  	.HAlignHeader(EHorizontalAlignment::HAlign_Right)
-	//  	.FillWidth(1.0f)
-	// );
-	UpdateDisplayColumn();
+	RegenerateHeader();
 	
 	ChildSlot
 	[
@@ -245,51 +225,57 @@ void SGeometryCollectionOutliner::Construct(const FArguments& InArgs)
 	];
 }
 
-void SGeometryCollectionOutliner::UpdateDisplayColumn()
+void SGeometryCollectionOutliner::RegenerateHeader()
 {
-	constexpr int32 CustomColumnIndex = 1;
 	constexpr int32 CustomFillWidth = 2.0f;
-	// first remove custom columns
-	HeaderRowWidget->RemoveColumn(SGeometryCollectionOutlinerColumnID::Custom);
+
+	HeaderRowWidget->ClearColumns();
+
+	HeaderRowWidget->AddColumn(
+		SHeaderRow::Column(SGeometryCollectionOutlinerColumnID::Bone)
+		.DefaultLabel(LOCTEXT("GCOutliner_Column_Bone", "Bone"))
+		.FillWidth(2.0f)
+	);
 
 	// then add the right customn one based on the selection
 	UOutlinerSettings* OutlinerSettings = GetMutableDefault<UOutlinerSettings>();
-	switch (OutlinerSettings->DisplayColumn)
+	switch (OutlinerSettings->ColumnMode)
 	{
-	case EOutlinerDisplayColumn::InitialState:
-		HeaderRowWidget->InsertColumn(
-			SHeaderRow::Column(SGeometryCollectionOutlinerColumnID::Custom)
-				.DefaultLabel(LOCTEXT("GCOutliner_Column_InitialState", "Initial State"))
-				.FillWidth(CustomFillWidth),
-			CustomColumnIndex
-		);
-		break;
-	case EOutlinerDisplayColumn::RelativeSize:
-		HeaderRowWidget->InsertColumn(
-			SHeaderRow::Column(SGeometryCollectionOutlinerColumnID::Custom)
+	case EOutlinerColumnMode::StateAndSize:
+		HeaderRowWidget->AddColumn(
+			SHeaderRow::Column(SGeometryCollectionOutlinerColumnID::RelativeSize)
 				.DefaultLabel(LOCTEXT("GCOutliner_Column_RelativeSize", "Relative Size"))
 				.HAlignHeader(EHorizontalAlignment::HAlign_Right)
-				.FillWidth(CustomFillWidth),
-			CustomColumnIndex
+				.FillWidth(CustomFillWidth)
 		);
-		break;
-	case EOutlinerDisplayColumn::Damages:
-		HeaderRowWidget->InsertColumn(
-			SHeaderRow::Column(SGeometryCollectionOutlinerColumnID::Custom)
-				.DefaultLabel(LOCTEXT("GCOutliner_Column_Damages", "Damages"))
-	 			.HAlignHeader(EHorizontalAlignment::HAlign_Center)
-	 			.FillWidth(CustomFillWidth),
-			CustomColumnIndex
+		HeaderRowWidget->AddColumn(
+			SHeaderRow::Column(SGeometryCollectionOutlinerColumnID::InitialState)
+				.DefaultLabel(LOCTEXT("GCOutliner_Column_InitialState", "Initial State"))
+				.HAlignHeader(EHorizontalAlignment::HAlign_Left)
+				.FillWidth(CustomFillWidth)
 		);
 		break;
-	default:
-		HeaderRowWidget->InsertColumn(
-			SHeaderRow::Column(SGeometryCollectionOutlinerColumnID::Custom)
-				.DefaultLabel(LOCTEXT("GCOutliner_Column_Custom", "-"))
-				.HAlignHeader(EHorizontalAlignment::HAlign_Right)
-				.FillWidth(CustomFillWidth),
-			CustomColumnIndex
+
+	case EOutlinerColumnMode::Damage:
+		HeaderRowWidget->AddColumn(
+			SHeaderRow::Column(SGeometryCollectionOutlinerColumnID::Damage)
+				.DefaultLabel(LOCTEXT("GCOutliner_Column_Damage", "Damage"))
+	 			.HAlignHeader(EHorizontalAlignment::HAlign_Right)
+	 			.FillWidth(CustomFillWidth)
 		);
+		HeaderRowWidget->AddColumn(
+			SHeaderRow::Column(SGeometryCollectionOutlinerColumnID::DamageThreshold)
+				.DefaultLabel(LOCTEXT("GCOutliner_Column_DamageThreshold", "Damage Threshold"))
+				 .HAlignHeader(EHorizontalAlignment::HAlign_Right)
+				 .FillWidth(CustomFillWidth)
+		);
+		HeaderRowWidget->AddColumn(
+			SHeaderRow::Column(SGeometryCollectionOutlinerColumnID::Broken)
+				.DefaultLabel(LOCTEXT("GCOutliner_Column_Broken", "Broken"))
+				 .HAlignHeader(EHorizontalAlignment::HAlign_Center)
+				 .FillWidth(1.0)
+		);
+		break;
 	}
 }
 
@@ -731,7 +717,9 @@ void FGeometryCollectionTreeItemBone::UpdateItemFromCollection()
 	
 	ItemColor = InvalidColor;
 	RelativeSize = 0;
+	Damage = 0;
 	DamageThreshold = 0;
+	Broken = false;
 	InitialState = INDEX_NONE;
 	
 	// Name / ItemText
@@ -744,7 +732,7 @@ void FGeometryCollectionTreeItemBone::UpdateItemFromCollection()
 
 	// Set color according to simulation type
 	const UGeometryCollectionComponent* Component = ParentComponentItem->GetComponent();
-	const FCollisionImpulseCollector* Collector = Component->GetRunTimeDataCollector();
+	const FDamageCollector* Collector = Component->GetRunTimeDataCollector();
 	const UGeometryCollection* RestCollection = Component->GetRestCollection();
 	if (RestCollection && IsValidChecked(RestCollection))
 	{
@@ -795,7 +783,13 @@ void FGeometryCollectionTreeItemBone::UpdateItemFromCollection()
 				RelativeSize = (*RelativeSizes)[ItemBoneIndex];
 			}
 
-			Damages = Collector? (*Collector)[ItemBoneIndex].Max : 0;  
+			if (Collector)
+			{
+				const FDamageCollector::FDamageData& DamageData = (*Collector)[ItemBoneIndex];
+				Damage = DamageData.MaxDamages;
+				DamageThreshold = DamageData.DamageThreshold;
+				Broken = DamageData.bIsBroken;
+			}
 		}
 	}
 }
@@ -863,7 +857,7 @@ TSharedRef<SWidget> FGeometryCollectionTreeItemBone::MakeDamagesColumnWidget() c
 			.HAlign(HAlign_Right)
 			[
 				SNew(STextBlock)
-				.Text(FText::AsNumber(Damages, &FormatOptions))
+				.Text(FText::AsNumber(Damage, &FormatOptions))
 				.ColorAndOpacity(ItemColor)
 			];
 }
@@ -881,6 +875,19 @@ TSharedRef<SWidget> FGeometryCollectionTreeItemBone::MakeDamageThresholdColumnWi
 			];
 }
 
+TSharedRef<SWidget> FGeometryCollectionTreeItemBone::MakeBrokenColumnWidget() const
+{
+	return SNew(SHorizontalBox)
+		+ SHorizontalBox::Slot()
+			.Padding(12.f, 0.f)
+			.HAlign(HAlign_Center)
+			[
+				SNew(STextBlock)
+				.Text(FText::FromString(Broken? TEXT("Broken"): TEXT("")))
+				.ColorAndOpacity(ItemColor)
+			];
+}
+
 TSharedRef<SWidget> FGeometryCollectionTreeItemBone::MakeInitialStateColumnWidget() const
 {
 	return SNew(SHorizontalBox)
diff --git a/Engine/Plugins/Experimental/ChaosEditor/Source/FractureEditor/Private/SGeometryCollectionOutliner.h b/Engine/Plugins/Experimental/ChaosEditor/Source/FractureEditor/Private/SGeometryCollectionOutliner.h
--- a/Engine/Plugins/Experimental/ChaosEditor/Source/FractureEditor/Private/SGeometryCollectionOutliner.h
+++ b/Engine/Plugins/Experimental/ChaosEditor/Source/FractureEditor/Private/SGeometryCollectionOutliner.h
@@ -36,11 +36,10 @@ enum class EOutlinerItemNameEnum : uint8
 };
 
 UENUM(BlueprintType)
-enum class EOutlinerDisplayColumn : uint8
+enum class EOutlinerColumnMode : uint8
 {
-	InitialState = 0		UMETA(DisplayName = "Initial State"),
-	RelativeSize = 1		UMETA(DisplayName = "Relative Size"),
-	Damages = 2				UMETA(Damages = "Damages"),
+	StateAndSize = 0		UMETA(DisplayName = "State And Size"),
+	Damage = 2				UMETA(DisplayName = "Damage"),
 };
 
 /** Settings for Outliner configuration. **/
@@ -61,8 +60,8 @@ public:
 	bool ColorByLevel;
 
 	/** the column to be display in the outliner */
-	UPROPERTY(EditAnywhere, Category = OutlinerSettings, meta = (DisplayName = "Display column"))
-	EOutlinerDisplayColumn DisplayColumn;
+	UPROPERTY(EditAnywhere, Category = OutlinerSettings, meta = (DisplayName = "Column Mode"))
+	EOutlinerColumnMode ColumnMode;
 };
 
 
@@ -145,9 +144,10 @@ public:
 		, ParentComponentItem(InParentComponentItem)
 		, ItemColor(FSlateColor::UseForeground())
 		, RelativeSize(0)
-		, Damages(0)
-		, DamageThreshold(0)
 		, InitialState(INDEX_NONE)
+		, Damage(0)
+		, DamageThreshold(0)
+		, Broken(false) 
 	{}
 
 	/** FGeometryCollectionTreeItem interface */
@@ -156,6 +156,7 @@ public:
 	TSharedRef<SWidget> MakeRelativeSizeColumnWidget() const;
 	TSharedRef<SWidget> MakeDamagesColumnWidget() const;
 	TSharedRef<SWidget> MakeDamageThresholdColumnWidget() const;
+	TSharedRef<SWidget> MakeBrokenColumnWidget() const;
 	TSharedRef<SWidget> MakeInitialStateColumnWidget() const;
 	TSharedRef<SWidget> MakeEmptyColumnWidget() const;
 	virtual void GetChildren(FGeometryCollectionTreeItemList& OutChildren) override;
@@ -180,9 +181,10 @@ private:
 	FSlateColor ItemColor;
 	FText ItemText;
 	float RelativeSize;
-	float Damages;
-	float DamageThreshold;
 	int32 InitialState;
+	float Damage;
+	float DamageThreshold;
+	bool Broken;
 };
 
 typedef TSharedPtr<class FGeometryCollectionTreeItemBone> FGeometryCollectionTreeItemBonePtr;
@@ -190,7 +192,13 @@ typedef TSharedPtr<class FGeometryCollectionTreeItemBone> FGeometryCollectionTre
 namespace SGeometryCollectionOutlinerColumnID
 {
 	const FName Bone("Bone");
-	const FName Custom("Custom");
+	// State and Size column mode
+	const FName RelativeSize("Relative Size");
+	const FName InitialState("Initial State");
+	// Damage Column Mode
+	const FName Damage("Damage");
+	const FName DamageThreshold("DamageThreshold");
+	const FName Broken("Broken");
 }
 
 class SGeometryCollectionOutlinerRow : public SMultiColumnTableRow<FGeometryCollectionTreeItemBonePtr>
@@ -232,7 +240,7 @@ public:
 	void Construct(const FArguments& InArgs);
 
 	void RegenerateItems();
-	void UpdateDisplayColumn();
+	void RegenerateHeader();
 
 	TSharedRef<ITableRow> MakeTreeRowWidget(FGeometryCollectionTreeItemPtr InInfo, const TSharedRef<STableViewBase>& OwnerTable);
 	TSharedRef<ITableRow> OnGeneratePinnedRowWidget(FGeometryCollectionTreeItemPtr InItem, const TSharedRef<STableViewBase>& InOwnerTable, bool bPinned);
diff --git a/Engine/Plugins/Experimental/ChaosEditor/Source/FractureEditor/Public/FractureEditorModeToolkit.h b/Engine/Plugins/Experimental/ChaosEditor/Source/FractureEditor/Public/FractureEditorModeToolkit.h
--- a/Engine/Plugins/Experimental/ChaosEditor/Source/FractureEditor/Public/FractureEditorModeToolkit.h
+++ b/Engine/Plugins/Experimental/ChaosEditor/Source/FractureEditor/Public/FractureEditorModeToolkit.h
@@ -213,6 +213,8 @@ private:
 	/** Callback for map changes. */
 	void HandleMapChanged(UWorld* NewWorld, EMapChangeType MapChangeType);
 
+	FReply RefreshOutliner();
+	
 private:
 	UFractureModalTool* ActiveTool;
 
diff --git a/Engine/Source/Runtime/Experimental/Chaos/Private/PhysicsProxy/GeometryCollectionPhysicsProxy.cpp b/Engine/Source/Runtime/Experimental/Chaos/Private/PhysicsProxy/GeometryCollectionPhysicsProxy.cpp
--- a/Engine/Source/Runtime/Experimental/Chaos/Private/PhysicsProxy/GeometryCollectionPhysicsProxy.cpp
+++ b/Engine/Source/Runtime/Experimental/Chaos/Private/PhysicsProxy/GeometryCollectionPhysicsProxy.cpp
@@ -1898,7 +1898,11 @@ void FGeometryCollectionPhysicsProxy::BufferPhysicsResults_External(Chaos::FDirt
 
 		TargetResults.Parent[TransformGroupIndex] = GameThreadCollection.Parent[TransformGroupIndex];
 		TargetResults.Transforms[TransformGroupIndex] = GameThreadCollection.Transform[TransformGroupIndex];
-		
+
+#if WITH_EDITORONLY_DATA
+		TargetResults.DamageInfo[TransformGroupIndex].Damage = 0.0f;
+		TargetResults.DamageInfo[TransformGroupIndex].DamageThreshold = 0.0f;
+#endif
 	}
 	
 	// TargetResults.GlobalTransforms; // no need, can be recomputed ?  
@@ -2107,7 +2111,9 @@ void FGeometryCollectionPhysicsProxy::BufferPhysicsResults_Internal(Chaos::FPBDR
 			PhysicsThreadCollection.Active[TransformGroupIndex] = !TargetResults.States[TransformGroupIndex].DisabledState;
 
 #if WITH_EDITORONLY_DATA
-			TargetResults.CollisionImpulses[TransformGroupIndex] = static_cast<float>(Handle->CollisionImpulse());
+			FGeometryCollectionResults::FDamageInfo& DamageInfo = TargetResults.DamageInfo[TransformGroupIndex];
+			DamageInfo.Damage = static_cast<float>(Handle->CollisionImpulse());
+			DamageInfo.DamageThreshold = static_cast<float>(Handle->Strain());
 #endif			
 		}    // end for
 	}        // STAT_CalcParticleToWorld scope
@@ -2187,11 +2193,12 @@ bool FGeometryCollectionPhysicsProxy::PullFromPhysicsState(const Chaos::FDirtyGe
 		}
 
 #if WITH_EDITORONLY_DATA
-		if (FCollisionImpulseCollector* Collector = FCollisionImpulseWatcher::GetInstance().Find(CollectorGuid))
+		if (FDamageCollector* Collector = FRuntimeDataCollector::GetInstance().Find(CollectorGuid))
 		{
 			for (int32 TransformGroupIndex = 0; TransformGroupIndex < NumTransforms; ++TransformGroupIndex)
 			{
-				Collector->SampleImpulse(TransformGroupIndex, TargetResults.CollisionImpulses[TransformGroupIndex]);
+				const FGeometryCollectionResults::FDamageInfo& DamageInfo = TargetResults.DamageInfo[TransformGroupIndex]; 
+				Collector->SampleDamage(TransformGroupIndex, DamageInfo.Damage, DamageInfo.DamageThreshold);
 			}
 		}
 #endif
@@ -3498,46 +3505,46 @@ void FGeometryCollectionPhysicsProxy::FieldForcesUpdateCallback(Chaos::FPBDRigid
 }
 
 
-void FCollisionImpulseCollector::Reset(int32 NumTransforms)
+void FDamageCollector::Reset(int32 NumTransforms)
 {
-	ImpulseData.Reset();
-	ImpulseData.SetNum(NumTransforms);
+	DamageData.Reset();
+	DamageData.SetNum(NumTransforms);
 }
 
-void FCollisionImpulseCollector::SampleImpulse(int32 TransformIndex, float Value)
+void FDamageCollector::SampleDamage(int32 TransformIndex, float Damage, float DamageThreshold)
 {
-	if (ImpulseData.IsValidIndex(TransformIndex))
+	if (DamageData.IsValidIndex(TransformIndex))
 	{
-		FImpulseData& Data = ImpulseData[TransformIndex];
-		Data.Min = FMath::Min(Data.Min, Value);
-		Data.Max = FMath::Max(Data.Max, Value);
-		Data.Count++;
+		FDamageData& Data = DamageData[TransformIndex];
+		Data.DamageThreshold = DamageThreshold;
+		Data.MaxDamages = FMath::Max(Data.MaxDamages, Damage);
+		Data.bIsBroken = Data.bIsBroken || (Damage > DamageThreshold);
 	}
 }
 
-FCollisionImpulseWatcher& FCollisionImpulseWatcher::GetInstance()
+FRuntimeDataCollector& FRuntimeDataCollector::GetInstance()
 {
-	static FCollisionImpulseWatcher Instance;
+	static FRuntimeDataCollector Instance;
 	return Instance;
 }
 
-void FCollisionImpulseWatcher::Clear()
+void FRuntimeDataCollector::Clear()
 {
 	Collectors.Reset();
 }
 
-FCollisionImpulseCollector* FCollisionImpulseWatcher::Find(const FGuid& Guid)
+FDamageCollector* FRuntimeDataCollector::Find(const FGuid& Guid)
 {
 	return Collectors.Find(Guid);
 }
 
-void FCollisionImpulseWatcher::AddCollector(const FGuid& Guid, int32 TransformNum)
+void FRuntimeDataCollector::AddCollector(const FGuid& Guid, int32 TransformNum)
 {
-	FCollisionImpulseCollector& Collector = Collectors.FindOrAdd(Guid);
+	FDamageCollector& Collector = Collectors.FindOrAdd(Guid);
 	Collector.Reset(TransformNum);
 }
 
-void FCollisionImpulseWatcher::RemoveCollector(const FGuid& Guid)
+void FRuntimeDataCollector::RemoveCollector(const FGuid& Guid)
 {
 	Collectors.Remove(Guid);
 	
diff --git a/Engine/Source/Runtime/Experimental/Chaos/Public/GeometryCollectionProxyData.h b/Engine/Source/Runtime/Experimental/Chaos/Public/GeometryCollectionProxyData.h
--- a/Engine/Source/Runtime/Experimental/Chaos/Public/GeometryCollectionProxyData.h
+++ b/Engine/Source/Runtime/Experimental/Chaos/Public/GeometryCollectionProxyData.h
@@ -162,7 +162,7 @@ public:
 		Transforms.SetNumUninitialized(NumTransforms);
 		Parent.SetNumUninitialized(NumTransforms);
 #if WITH_EDITORONLY_DATA
-		CollisionImpulses.SetNumUninitialized(NumTransforms);
+		DamageInfo.SetNumUninitialized(NumTransforms);
 #endif		
 	}
 
@@ -172,9 +172,9 @@ public:
 		int16 DisabledState: 1;
 		int16 HasInternalClusterParent: 1;
 		int16 DynamicInternalClusterParent: 1;
-		// 6 bits left
+		// 5 bits left
 	};
-	
+
 	Chaos::FReal SolverDt;
 	TArray<FState> States;
 	TArray<FMatrix> GlobalTransforms;
@@ -187,8 +187,14 @@ public:
 	TArray<int32> Parent;
 
 #if WITH_EDITORONLY_DATA
+	struct FDamageInfo
+	{
+		float Damage = 0;
+		float DamageThreshold = 0;
+	};
+	
 	// use to display impulse statistics in editor
-	TArray<float> CollisionImpulses;
+	TArray<FDamageInfo> DamageInfo;
 #endif
 	
 	bool IsObjectDynamic;
diff --git a/Engine/Source/Runtime/Experimental/Chaos/Public/PhysicsProxy/GeometryCollectionPhysicsProxy.h b/Engine/Source/Runtime/Experimental/Chaos/Public/PhysicsProxy/GeometryCollectionPhysicsProxy.h
--- a/Engine/Source/Runtime/Experimental/Chaos/Public/PhysicsProxy/GeometryCollectionPhysicsProxy.h
+++ b/Engine/Source/Runtime/Experimental/Chaos/Public/PhysicsProxy/GeometryCollectionPhysicsProxy.h
@@ -385,42 +385,42 @@ private:
 };
 
 /**
- * the collision impulse watcher collects runtime data about
- * collision impulses on each piece of the geometry collection 
+ * the watcher collects runtime data about
+ * damage on each piece of the geometry collection 
  */
-struct CHAOS_API FCollisionImpulseCollector
+struct CHAOS_API FDamageCollector
 {
 public:
-	struct FImpulseData
+	struct FDamageData
 	{
-		float Min = TNumericLimits<float>::Max();
-		float Max = 0;
-		int32 Count = 0;
+		float DamageThreshold= 0;
+		float MaxDamages = 0;
+		bool  bIsBroken = false;
 	};
 
 	void Reset(int32 NumTransforms);
-	int32 Num() const { return ImpulseData.Num(); }
-	const FImpulseData& operator[](int32 TransformIndex) const { return ImpulseData[TransformIndex]; }
-	void SampleImpulse(int32 TransformIndex, float Value);
+	int32 Num() const { return DamageData.Num(); }
+	const FDamageData& operator[](int32 TransformIndex) const { return DamageData[TransformIndex]; }
+	void SampleDamage(int32 TransformIndex, float Damage, float DamageThreshold);
 	
 private:
-	TArray<FImpulseData> ImpulseData;
+	TArray<FDamageData> DamageData;
 };
 
-struct CHAOS_API FCollisionImpulseWatcher
+struct CHAOS_API FRuntimeDataCollector
 {
 public:
 	void Clear();
 	void AddCollector(const FGuid& Guid, int32 TransformNum);
 	void RemoveCollector(const FGuid& Guid);
 
-	FCollisionImpulseCollector* Find(const FGuid& Guid);
+	FDamageCollector* Find(const FGuid& Guid);
 
-	static FCollisionImpulseWatcher& GetInstance();
+	static FRuntimeDataCollector& GetInstance();
 	
 private:
 	// collectors by geometry collection Guids
-	TMap<FGuid,FCollisionImpulseCollector> Collectors;
+	TMap<FGuid,FDamageCollector> Collectors;
 };
 
 CHAOS_API TUniquePtr<Chaos::FTriangleMesh> CreateTriangleMesh(const int32 FaceStart,const int32 FaceCount,const TManagedArray<bool>& Visible,const TManagedArray<FIntVector>& Indices, bool bRotateWinding = true);
diff --git a/Engine/Source/Runtime/Experimental/GeometryCollectionEngine/Private/GeometryCollection/GeometryCollectionComponent.cpp b/Engine/Source/Runtime/Experimental/GeometryCollectionEngine/Private/GeometryCollection/GeometryCollectionComponent.cpp
--- a/Engine/Source/Runtime/Experimental/GeometryCollectionEngine/Private/GeometryCollection/GeometryCollectionComponent.cpp
+++ b/Engine/Source/Runtime/Experimental/GeometryCollectionEngine/Private/GeometryCollection/GeometryCollectionComponent.cpp
@@ -2228,11 +2228,11 @@ void UGeometryCollectionComponent::RegisterAndInitializePhysicsProxy()
 	CollectorGuid = RunTimeDataCollectionGuid;
 	if (bEnableRunTimeDataCollection && RestCollection)
 	{
-		FCollisionImpulseWatcher::GetInstance().AddCollector(CollectorGuid, RestCollection->NumElements(FGeometryCollection::TransformGroup));
+		FRuntimeDataCollector::GetInstance().AddCollector(CollectorGuid, RestCollection->NumElements(FGeometryCollection::TransformGroup));
 	}
 	else
 	{
-		FCollisionImpulseWatcher::GetInstance().RemoveCollector(CollectorGuid);
+		FRuntimeDataCollector::GetInstance().RemoveCollector(CollectorGuid);
 	}
 #endif
 	PhysicsProxy = new FGeometryCollectionPhysicsProxy(this, *DynamicCollection, SimulationParameters, InitialSimFilter, InitialQueryFilter, CollectorGuid);
@@ -2244,9 +2244,9 @@ void UGeometryCollectionComponent::RegisterAndInitializePhysicsProxy()
 }
 
 #if WITH_EDITORONLY_DATA
-const FCollisionImpulseCollector* UGeometryCollectionComponent::GetRunTimeDataCollector() const
+const FDamageCollector* UGeometryCollectionComponent::GetRunTimeDataCollector() const
 {
-	return FCollisionImpulseWatcher::GetInstance().Find(RunTimeDataCollectionGuid);
+	return FRuntimeDataCollector::GetInstance().Find(RunTimeDataCollectionGuid);
 }
 #endif
 
diff --git a/Engine/Source/Runtime/Experimental/GeometryCollectionEngine/Public/GeometryCollection/GeometryCollectionComponent.h b/Engine/Source/Runtime/Experimental/GeometryCollectionEngine/Public/GeometryCollection/GeometryCollectionComponent.h
--- a/Engine/Source/Runtime/Experimental/GeometryCollectionEngine/Public/GeometryCollection/GeometryCollectionComponent.h
+++ b/Engine/Source/Runtime/Experimental/GeometryCollectionEngine/Public/GeometryCollection/GeometryCollectionComponent.h
@@ -35,7 +35,7 @@ class UChaosPhysicalMaterial;
 class AChaosSolverActor;
 struct FGeometryCollectionEmbeddedExemplar;
 class UInstancedStaticMeshComponent;
-struct FCollisionImpulseCollector;
+struct FDamageCollector;
 
 DECLARE_DYNAMIC_MULTICAST_DELEGATE_OneParam(FOnChaosBreakEvent, const FChaosBreakEvent&, BreakEvent);
 
@@ -710,7 +710,7 @@ public:
 #endif
 
 #if WITH_EDITORONLY_DATA
-	const FCollisionImpulseCollector* GetRunTimeDataCollector() const;
+	const FDamageCollector* GetRunTimeDataCollector() const;
 #endif
 	
 	// #todo should this only be available in editor?
