Fracture tool : fix localization string conflict

diff --git a/Engine/Plugins/Experimental/ChaosEditor/Source/FractureEditor/Private/SGeometryCollectionOutliner.cpp b/Engine/Plugins/Experimental/ChaosEditor/Source/FractureEditor/Private/SGeometryCollectionOutliner.cpp
--- a/Engine/Plugins/Experimental/ChaosEditor/Source/FractureEditor/Private/SGeometryCollectionOutliner.cpp
+++ b/Engine/Plugins/Experimental/ChaosEditor/Source/FractureEditor/Private/SGeometryCollectionOutliner.cpp
@@ -265,8 +265,8 @@ void SGeometryCollectionOutliner::RegenerateHeader()
 	case EOutlinerColumnMode::Damage:
 		HeaderRowWidget->AddColumn(
 			SHeaderRow::Column(SGeometryCollectionOutlinerColumnID::Damage)
-				.DefaultLabel(LOCTEXT("GCOutliner_Column_Damage", "Max damage"))
-				.DefaultTooltip(LOCTEXT("GCOutliner_Column_Damage_ToolTip", "Maximum amount of damage recorded ( through collision )"))
+				.DefaultLabel(LOCTEXT("GCOutliner_Column_MaxDamage", "Max damage"))
+				.DefaultTooltip(LOCTEXT("GCOutliner_Column_MaxDamage_ToolTip", "Maximum amount of damage recorded ( through collision )"))
 	 			.HAlignHeader(EHorizontalAlignment::HAlign_Center)
 	 			.FillWidth(CustomFillWidth)
 		);
