Add imported collision info to the outliner

diff --git a/Engine/Plugins/Experimental/ChaosEditor/Source/FractureEditor/Private/SGeometryCollectionOutliner.cpp b/Engine/Plugins/Experimental/ChaosEditor/Source/FractureEditor/Private/SGeometryCollectionOutliner.cpp
--- a/Engine/Plugins/Experimental/ChaosEditor/Source/FractureEditor/Private/SGeometryCollectionOutliner.cpp
+++ b/Engine/Plugins/Experimental/ChaosEditor/Source/FractureEditor/Private/SGeometryCollectionOutliner.cpp
@@ -25,7 +25,7 @@ static FText GettextFromInitialDynamicState(int32 InitialDynamicState)
 {
 	switch (InitialDynamicState)
 	{
-	case 0: return NSLOCTEXT("Fracture", "FractureToolInitialDynamicStateNoOverride", "No Override");
+	case 0: return NSLOCTEXT("Fracture", "FractureToolInitialDynamicStateNoOverride", "-");
 	case 1: return NSLOCTEXT("Fracture", "FractureToolInitialDynamicStateSleeping", "Sleeping");
 	case 2: return NSLOCTEXT("Fracture", "FractureToolInitialDynamicStateKinematic", "Kinematic");
 	case 3: return NSLOCTEXT("Fracture", "FractureToolInitialDynamicStateStatic", "Static");
@@ -199,6 +199,8 @@ TSharedRef<SWidget> SGeometryCollectionOutlinerRow::GenerateWidgetForColumn(cons
 		return Item->MakePostBreakTimecolumnWidget();
 	if (ColumnName == SGeometryCollectionOutlinerColumnID::RemovalTime)
 		return Item->MakeRemovalTimeColumnWidget();
+	if (ColumnName == SGeometryCollectionOutlinerColumnID::ImportedCollisions)
+		return Item->MakeImportedCollisionsColumnWidget();
 	return Item->MakeEmptyColumnWidget();
 }
 
@@ -302,6 +304,15 @@ void SGeometryCollectionOutliner::RegenerateHeader()
 		);
 		break;
 		
+	case EOutlinerColumnMode::Collision:
+		HeaderRowWidget->AddColumn(
+			SHeaderRow::Column(SGeometryCollectionOutlinerColumnID::ImportedCollisions)
+				.DefaultLabel(LOCTEXT("GCOutliner_Column_ImportedCollisions", "Imported Collisions"))
+				.DefaultTooltip(LOCTEXT("GCOutliner_Column_ImportedCollisions_ToolTip", "Status of imported Collision [available, used, nothing])"))
+				.HAlignHeader(EHorizontalAlignment::HAlign_Center)
+				.FillWidth(CustomFillWidth)
+		);
+		break;
 	}
 }
 
@@ -749,6 +760,8 @@ void FGeometryCollectionTreeItemBone::UpdateItemFromCollection()
 	InitialState = INDEX_NONE;
 	RemoveOnBreakAvailable = false;
 	RemoveOnBreak = FVector4f{-1};
+	ImportedCollisionsAvailable = false;
+	ImportedCollisionsUsed = false;
 	
 	// Name / ItemText
 	UOutlinerSettings* OutlinerSettings = GetMutableDefault<UOutlinerSettings>();
@@ -768,6 +781,9 @@ void FGeometryCollectionTreeItemBone::UpdateItemFromCollection()
 		const TManagedArray<int32>& SimulationType = GeometryCollectionPtr->SimulationType;
 		const TManagedArray<float>* RelativeSizes = GeometryCollectionPtr->FindAttribute<float>("Size", FTransformCollection::TransformGroup);
 		const TManagedArray<FVector4f>* RemoveOnBreakArray = GeometryCollectionPtr->FindAttribute<FVector4f>("RemoveOnBreak", FTransformCollection::TransformGroup);
+
+		using FImplicitGeom = FGeometryDynamicCollection::FSharedImplicit;
+		const TManagedArray<FImplicitGeom>* ExternalCollisions = GeometryCollectionPtr->FindAttribute<FImplicitGeom>("ExternalCollisions", FTransformCollection::TransformGroup);
 		
 		if (ensure(ItemBoneIndex >= 0 && ItemBoneIndex < SimulationType.Num()))
 		{
@@ -825,6 +841,12 @@ void FGeometryCollectionTreeItemBone::UpdateItemFromCollection()
 			{
 				RemoveOnBreak = (*RemoveOnBreakArray)[ItemBoneIndex];
 			}
+
+			ImportedCollisionsUsed = RestCollection->bImportCollisionFromSource;
+			if (ExternalCollisions)
+			{
+				ImportedCollisionsAvailable = (*ExternalCollisions)[ItemBoneIndex] != nullptr;
+			}
 		}
 	}
 }
@@ -962,6 +984,31 @@ TSharedRef<SWidget> FGeometryCollectionTreeItemBone::MakeRemovalTimeColumnWidget
 		];
 }
 
+TSharedRef<SWidget> FGeometryCollectionTreeItemBone::MakeImportedCollisionsColumnWidget() const
+{
+	FText ImportedCollisionText = LOCTEXT("GCOutliner_ImportedCollision_NotAvailable", "-");
+	if (ImportedCollisionsAvailable)
+	{
+		if (ImportedCollisionsUsed)
+		{
+			ImportedCollisionText = LOCTEXT("GCOutliner_ImportedCollision_Used", "Used");
+		}
+		else
+		{
+			ImportedCollisionText = LOCTEXT("GCOutliner_ImportedCollision_Available", "Available");
+		}
+	}
+	return SNew(SHorizontalBox)
+	+ SHorizontalBox::Slot()
+		.Padding(12.f, 0.f)
+		.HAlign(HAlign_Center)
+		[
+			SNew(STextBlock)
+			.Text(ImportedCollisionText)
+			.ColorAndOpacity(ItemColor)
+		];
+}
+
 TSharedRef<SWidget> FGeometryCollectionTreeItemBone::MakeInitialStateColumnWidget() const
 {
 	return SNew(SHorizontalBox)
diff --git a/Engine/Plugins/Experimental/ChaosEditor/Source/FractureEditor/Private/SGeometryCollectionOutliner.h b/Engine/Plugins/Experimental/ChaosEditor/Source/FractureEditor/Private/SGeometryCollectionOutliner.h
--- a/Engine/Plugins/Experimental/ChaosEditor/Source/FractureEditor/Private/SGeometryCollectionOutliner.h
+++ b/Engine/Plugins/Experimental/ChaosEditor/Source/FractureEditor/Private/SGeometryCollectionOutliner.h
@@ -41,6 +41,7 @@ enum class EOutlinerColumnMode : uint8
 	StateAndSize = 0		UMETA(DisplayName = "State And Size"),
 	Damage = 1				UMETA(DisplayName = "Damage"),
 	Removal = 2				UMETA(DisplayName = "Removal"),
+	Collision = 3			UMETA(DisplayName = "Collision"),
 };
 
 /** Settings for Outliner configuration. **/
@@ -151,6 +152,8 @@ public:
 		, Broken(false)
 		, RemoveOnBreakAvailable(false)
 		, RemoveOnBreak(-1)
+		, ImportedCollisionsAvailable(false)
+		, ImportedCollisionsUsed(false)
 	{}
 
 	/** FGeometryCollectionTreeItem interface */
@@ -163,6 +166,7 @@ public:
 	TSharedRef<SWidget> MakeInitialStateColumnWidget() const;
 	TSharedRef<SWidget> MakePostBreakTimecolumnWidget() const;
 	TSharedRef<SWidget> MakeRemovalTimeColumnWidget() const;
+	TSharedRef<SWidget> MakeImportedCollisionsColumnWidget() const;
 	TSharedRef<SWidget> MakeEmptyColumnWidget() const;
 	virtual void GetChildren(FGeometryCollectionTreeItemList& OutChildren) override;
 	virtual int32 GetBoneIndex() const override { return BoneIndex; }
@@ -192,6 +196,8 @@ private:
 	bool Broken;
 	bool RemoveOnBreakAvailable;
 	FVector4f RemoveOnBreak;
+	bool ImportedCollisionsAvailable;
+	bool ImportedCollisionsUsed;
 };
 
 typedef TSharedPtr<class FGeometryCollectionTreeItemBone> FGeometryCollectionTreeItemBonePtr;
@@ -209,6 +215,8 @@ namespace SGeometryCollectionOutlinerColumnID
 	// Removal Column Mode
 	const FName PostBreakTime("PostBreakTime");
 	const FName RemovalTime("RemovalTime");
+	// Collision Column Mode
+	const FName ImportedCollisions("ImportedCollisions");
 }
 
 class SGeometryCollectionOutlinerRow : public SMultiColumnTableRow<FGeometryCollectionTreeItemBonePtr>
