expand progresscancel usage to fracture uv tools

diff --git a/Engine/Plugins/Experimental/ChaosEditor/Source/FractureEditor/Private/FractureToolBackgroundTask.h b/Engine/Plugins/Experimental/ChaosEditor/Source/FractureEditor/Private/FractureToolBackgroundTask.h
--- a/Engine/Plugins/Experimental/ChaosEditor/Source/FractureEditor/Private/FractureToolBackgroundTask.h
+++ b/Engine/Plugins/Experimental/ChaosEditor/Source/FractureEditor/Private/FractureToolBackgroundTask.h
@@ -14,27 +14,53 @@ namespace Fracture
 {
 
 // Base class for background operators that update geometry collections (e.g. to fracture in a background thread)
-class FGeometryCollectionOperator : public UE::Geometry::TGenericDataOperator<FGeometryCollection>
+template<typename ResultType>
+class TGeometryCollectionOperator : public UE::Geometry::TGenericDataOperator<ResultType>
 {
 public:
 	int ResultGeometryIndex = -1;
 	TUniquePtr<FGeometryCollection> CollectionCopy;
 
-	FGeometryCollectionOperator(const FGeometryCollection& SourceCollection)
+	TGeometryCollectionOperator(const FGeometryCollection& SourceCollection)
 	{
 		CollectionCopy = MakeUnique<FGeometryCollection>();
 		CollectionCopy->CopyMatchingAttributesFrom(SourceCollection, nullptr);
 	}
 
-	virtual ~FGeometryCollectionOperator() = default;
+	virtual ~TGeometryCollectionOperator() = default;
 
 	virtual int GetResultGeometryIndex()
 	{
 		return ResultGeometryIndex;
 	}
 
-	// Post-process the geometry collection on success -- by default, clears proximity data
-	virtual void OnSuccess(FGeometryCollection& Collection)
+	// Post-process the geometry collection on success
+	virtual void OnSuccess(ResultType& Collection)
+	{
+
+	}
+
+	// For operators where there is no geometry index to report, set a generic success value
+	virtual void SetSuccessIndex()
+	{
+		ResultGeometryIndex = 0;
+	}
+};
+
+// Most operators will also have FGeometryCollection as their result
+using FGeometryCollectionOperator = TGeometryCollectionOperator<FGeometryCollection>;
+
+// Fracture-specific operators also clear proximity on success
+class FGeometryCollectionFractureOperator : public FGeometryCollectionOperator
+{
+public:
+	FGeometryCollectionFractureOperator(const FGeometryCollection& SourceCollection) : FGeometryCollectionOperator(SourceCollection)
+	{
+	}
+
+	virtual ~FGeometryCollectionFractureOperator() = default;
+
+	virtual void OnSuccess(FGeometryCollection& Collection) override
 	{
 		// Invalidate proximity
 		if (Collection.HasAttribute("Proximity", FGeometryCollection::GeometryGroup))
@@ -44,9 +70,10 @@ public:
 	}
 };
 
+
 // Run a blocking geometry collection op, but with a responsive cancel option
-template<class GeometryCollectionOpType>
-int RunCancellableGeometryCollectionOp(FGeometryCollection& ToUpdate, TUniquePtr<GeometryCollectionOpType>&& NewOp, FText DefaultMessage, float DialogDelay = .5)
+template<class GeometryCollectionOpType, typename ResultType>
+int RunCancellableGeometryCollectionOpGeneric(ResultType& ToUpdate, TFunctionRef<void(ResultType&, ResultType&)> AssignResult, TUniquePtr<GeometryCollectionOpType>&& NewOp, FText DefaultMessage, float DialogDelay = .5)
 {
 	using FGeometryCollectionTask = UE::Geometry::TModelingOpTask<GeometryCollectionOpType>;
 	using FExecuter = UE::Geometry::FAsyncTaskExecuterWithProgressCancel<FGeometryCollectionTask>;
@@ -59,17 +86,17 @@ int RunCancellableGeometryCollectionOp(FGeometryCollection& ToUpdate, TUniquePtr
 	bool bSuccess = false;
 	while (true)
 	{
-		if (BackgroundTask->IsDone())
-		{
-			bSuccess = !BackgroundTask->GetTask().IsAborted();
-			break;
-		}
 		if (SlowTask.ShouldCancel())
 		{
 			// Release ownership to the TDeleterTask that is spawned by CancelAndDelete()
 			BackgroundTask.Release()->CancelAndDelete();
 			break;
 		}
+		if (BackgroundTask->IsDone())
+		{
+			bSuccess = !BackgroundTask->GetTask().IsAborted();
+			break;
+		}
 		FPlatformProcess::Sleep(.2); // SlowTask::ShouldCancel will throttle any updates faster than .2 seconds
 		float ProgressFrac;
 		FText ProgressMessage;
@@ -90,13 +117,15 @@ int RunCancellableGeometryCollectionOp(FGeometryCollection& ToUpdate, TUniquePtr
 	{
 		check(BackgroundTask != nullptr && BackgroundTask->IsDone());
 		TUniquePtr<GeometryCollectionOpType> Op = BackgroundTask->GetTask().ExtractOperator();
-		TUniquePtr<FGeometryCollection> Result = Op->ExtractResult();
+
+		TUniquePtr<ResultType> Result = Op->ExtractResult();
 		if (!Result.IsValid())
 		{
 			return -1;
 		}
 
-		ToUpdate.CopyMatchingAttributesFrom(*Result, nullptr);
+		AssignResult(ToUpdate, *Result);
+
 		Op->OnSuccess(ToUpdate);
 		return Op->GetResultGeometryIndex();
 	}
@@ -104,5 +133,15 @@ int RunCancellableGeometryCollectionOp(FGeometryCollection& ToUpdate, TUniquePtr
 	return -1;
 }
 
+template<class GeometryCollectionOpType>
+int RunCancellableGeometryCollectionOp(FGeometryCollection& ToUpdate, TUniquePtr<GeometryCollectionOpType>&& NewOp, FText DefaultMessage, float DialogDelay = .5)
+{
+	return RunCancellableGeometryCollectionOpGeneric<GeometryCollectionOpType, FGeometryCollection>(
+		ToUpdate, [](FGeometryCollection& ToUpdateIn, FGeometryCollection& ResultIn)
+		{
+			ToUpdateIn.CopyMatchingAttributesFrom(ResultIn, nullptr);
+		}, MoveTemp(NewOp), DefaultMessage, DialogDelay);
+}
+
 } // end namespace UE::Fracture
 } // end namespace UE
\ No newline at end of file
diff --git a/Engine/Plugins/Experimental/ChaosEditor/Source/FractureEditor/Private/FractureToolBrick.cpp b/Engine/Plugins/Experimental/ChaosEditor/Source/FractureEditor/Private/FractureToolBrick.cpp
--- a/Engine/Plugins/Experimental/ChaosEditor/Source/FractureEditor/Private/FractureToolBrick.cpp
+++ b/Engine/Plugins/Experimental/ChaosEditor/Source/FractureEditor/Private/FractureToolBrick.cpp
@@ -376,10 +376,10 @@ void UFractureToolBrick::AddBoxEdges(const FVector& Min, const FVector& Max)
 	Edges.Emplace(MakeTuple(FVector(Min.X, Max.Y, Max.Z), Max));
 }
 
-class FCellsFractureOp : public FGeometryCollectionOperator
+class FCellsFractureOp : public FGeometryCollectionFractureOperator
 {
 public:
-	FCellsFractureOp(const FGeometryCollection& SourceCollection) : FGeometryCollectionOperator(SourceCollection)
+	FCellsFractureOp(const FGeometryCollection& SourceCollection) : FGeometryCollectionFractureOperator(SourceCollection)
 	{}
 
 	virtual ~FCellsFractureOp() = default;
diff --git a/Engine/Plugins/Experimental/ChaosEditor/Source/FractureEditor/Private/FractureToolCutter.cpp b/Engine/Plugins/Experimental/ChaosEditor/Source/FractureEditor/Private/FractureToolCutter.cpp
--- a/Engine/Plugins/Experimental/ChaosEditor/Source/FractureEditor/Private/FractureToolCutter.cpp
+++ b/Engine/Plugins/Experimental/ChaosEditor/Source/FractureEditor/Private/FractureToolCutter.cpp
@@ -354,10 +354,10 @@ void UFractureToolVoronoiCutterBase::FractureContextChanged()
 }
 
 
-class FVoronoiFractureOp : public FGeometryCollectionOperator
+class FVoronoiFractureOp : public FGeometryCollectionFractureOperator
 {
 public:
-	FVoronoiFractureOp(const FGeometryCollection& SourceCollection) : FGeometryCollectionOperator(SourceCollection)
+	FVoronoiFractureOp(const FGeometryCollection& SourceCollection) : FGeometryCollectionFractureOperator(SourceCollection)
 	{}
 
 	virtual ~FVoronoiFractureOp() = default;
diff --git a/Engine/Plugins/Experimental/ChaosEditor/Source/FractureEditor/Private/FractureToolMeshCut.cpp b/Engine/Plugins/Experimental/ChaosEditor/Source/FractureEditor/Private/FractureToolMeshCut.cpp
--- a/Engine/Plugins/Experimental/ChaosEditor/Source/FractureEditor/Private/FractureToolMeshCut.cpp
+++ b/Engine/Plugins/Experimental/ChaosEditor/Source/FractureEditor/Private/FractureToolMeshCut.cpp
@@ -199,10 +199,10 @@ void UFractureToolMeshCut::FractureContextChanged()
 	}
 }
 
-class FMeshFractureOp : public FGeometryCollectionOperator
+class FMeshFractureOp : public FGeometryCollectionFractureOperator
 {
 public:
-	FMeshFractureOp(const FGeometryCollection& SourceCollection) : FGeometryCollectionOperator(SourceCollection)
+	FMeshFractureOp(const FGeometryCollection& SourceCollection) : FGeometryCollectionFractureOperator(SourceCollection)
 	{}
 
 	virtual ~FMeshFractureOp() = default;
@@ -289,7 +289,7 @@ int32 UFractureToolMeshCut::ExecuteFracture(const FFractureToolContext& Fracture
 		MeshCutOp->CuttingMesh = ConvertMeshDescriptionToCuttingDynamicMesh(MeshDescription, FractureContext.GetGeometryCollection()->NumUVLayers());
 
 		int Result = RunCancellableGeometryCollectionOp<FMeshFractureOp>(*(FractureContext.GetGeometryCollection()),
-			MoveTemp(MeshCutOp), LOCTEXT("ComputingVoronoiFractureMessage", "Computing Mesh Fracture"));
+			MoveTemp(MeshCutOp), LOCTEXT("ComputingMeshFractureMessage", "Computing Mesh Fracture"));
 		return Result;
 	}
 
diff --git a/Engine/Plugins/Experimental/ChaosEditor/Source/FractureEditor/Private/FractureToolUV.cpp b/Engine/Plugins/Experimental/ChaosEditor/Source/FractureEditor/Private/FractureToolUV.cpp
--- a/Engine/Plugins/Experimental/ChaosEditor/Source/FractureEditor/Private/FractureToolUV.cpp
+++ b/Engine/Plugins/Experimental/ChaosEditor/Source/FractureEditor/Private/FractureToolUV.cpp
@@ -19,12 +19,14 @@
 #include "Misc/ScopedSlowTask.h"
 #include "ScopedTransaction.h"
 
+#include "FractureToolBackgroundTask.h"
+
 #include "FractureAutoUV.h"
 
 #define LOCTEXT_NAMESPACE "FractureToolAutoUV"
 
 using namespace UE::Geometry;
-
+using namespace UE::Fracture;
 
 void UFractureAutoUVSettings::SetNumUVChannels(int32 NumUVChannels)
 {
@@ -395,6 +397,36 @@ namespace
 	}
 }
 
+ class FLayoutUVOp : public FGeometryCollectionOperator
+ {
+ public:
+ 	FLayoutUVOp(const FGeometryCollection& SourceCollection) : FGeometryCollectionOperator(SourceCollection)
+ 	{}
+
+ 	virtual ~FLayoutUVOp() = default;
+
+	int32 UVLayer;
+	int32 OutputRes;
+	int32 GutterSize;
+	UE::PlanarCut::EUseMaterials UseMaterialIDs;
+	TArray<int32> MaterialIDs;
+
+ 	// TGenericDataOperator interface:
+ 	virtual void CalculateResult(FProgressCancel* Progress) override
+ 	{
+		bool bLayoutSuccess =
+			UE::PlanarCut::UVLayout(UVLayer, *CollectionCopy, OutputRes, GutterSize,
+				UseMaterialIDs, MaterialIDs, true, Progress);
+		
+		if (bLayoutSuccess)
+		{
+			SetSuccessIndex();
+			SetResult(MoveTemp(CollectionCopy));
+		}
+ 	}
+ };
+
+
 void UFractureToolAutoUV::LayoutUVs()
 {
 	TSet<UGeometryCollectionComponent*> GeomCompSelection;
@@ -415,22 +447,58 @@ void UFractureToolAutoUV::LayoutUVs()
 
 bool UFractureToolAutoUV::LayoutUVsForComponent(UGeometryCollectionComponent* Component)
 {
-	int32 OutputRes = (int32)AutoUVSettings->Resolution;
-	TArray<int32> EmptyMaterialIDs;
+	FGeometryCollection& Collection = *Component->GetRestCollection()->GetGeometryCollection();
 
-	int32 UVLayer = AutoUVSettings->GetSelectedChannelIndex();
+	TUniquePtr<FLayoutUVOp> LayoutOp = MakeUnique<FLayoutUVOp>(Collection);
+	LayoutOp->OutputRes = (int32)AutoUVSettings->Resolution;
+	TArray<int32> EmptyMaterialIDs;
+	LayoutOp->MaterialIDs = AutoUVSettings->TargetMaterialIDs == ETargetMaterialIDs::OddIDs ? EmptyMaterialIDs : AutoUVSettings->MaterialIDs;
 
-	UE::PlanarCut::EUseMaterials UseMaterialIDs = GetUseMaterials(AutoUVSettings->TargetMaterialIDs);
+	LayoutOp->OutputRes = (int32)AutoUVSettings->Resolution;
+	LayoutOp->UVLayer = AutoUVSettings->GetSelectedChannelIndex();
+	LayoutOp->GutterSize = AutoUVSettings->GutterSize;
+	LayoutOp->UseMaterialIDs = GetUseMaterials(AutoUVSettings->TargetMaterialIDs);
 
-	bool bLayoutSuccess =
-		UE::PlanarCut::UVLayout(UVLayer, *Component->GetRestCollection()->GetGeometryCollection(), OutputRes, AutoUVSettings->GutterSize,
-			UseMaterialIDs,
-			AutoUVSettings->TargetMaterialIDs == ETargetMaterialIDs::OddIDs ? EmptyMaterialIDs : AutoUVSettings->MaterialIDs);
-	
-	return bLayoutSuccess;
+	int Result = RunCancellableGeometryCollectionOp<FLayoutUVOp>(Collection,
+		MoveTemp(LayoutOp), LOCTEXT("ComputingUVLayoutMessage", "Computing UV Layout"));
+	return Result != -1;
 }
 
 
+class FBakeTextureOp : public TGeometryCollectionOperator<UE::Geometry::TImageBuilder<FVector4f>>
+{
+public:
+	FBakeTextureOp(const FGeometryCollection& SourceCollection) : TGeometryCollectionOperator<UE::Geometry::TImageBuilder<FVector4f>>(SourceCollection)
+	{}
+
+	virtual ~FBakeTextureOp() = default;
+
+	int32 UVLayer;
+	int32 GutterSize;
+	FIndex4i Attributes;
+	UE::PlanarCut::FTextureAttributeSettings AttribSettings;
+	TUniquePtr<UE::Geometry::TImageBuilder<FVector4f>> ImageBuilder;
+	UE::PlanarCut::EUseMaterials UseMaterialIDs;
+	TArray<int32> MaterialIDs;
+
+	// TGenericDataOperator interface:
+	virtual void CalculateResult(FProgressCancel* Progress) override
+	{
+		bool bBakeSuccess = UE::PlanarCut::TextureInternalSurfaces(
+			UVLayer, *CollectionCopy, GutterSize, Attributes, AttribSettings, *ImageBuilder,
+			UseMaterialIDs, MaterialIDs, Progress);
+
+		if (!bBakeSuccess || (Progress && Progress->Cancelled()))
+		{
+			return;
+		}
+
+		SetResult(MoveTemp(ImageBuilder));
+		SetSuccessIndex();
+	}
+};
+
+
 void UFractureToolAutoUV::BakeTexture()
 {
 	TSet<UGeometryCollectionComponent*> GeomCompSelection;
@@ -443,30 +511,32 @@ void UFractureToolAutoUV::BakeTexture()
 	}
 }
 
-void UFractureToolAutoUV::BakeTextureForComponent(UGeometryCollectionComponent* Component, TFunction<void(int32, const FText&)> Progress)
+void UFractureToolAutoUV::BakeTextureForComponent(UGeometryCollectionComponent* Component)
 {
+	using FImageType = UE::Geometry::TImageBuilder<FVector4f>;
+
 	FGeometryCollection& Collection = *Component->GetRestCollection()->GetGeometryCollection();
 
-	int32 OutputRes = (int32)AutoUVSettings->Resolution;
+	TUniquePtr<FBakeTextureOp> BakeOp = MakeUnique<FBakeTextureOp>(Collection);
 	TArray<int32> EmptyMaterialIDs;
+	BakeOp->MaterialIDs = AutoUVSettings->TargetMaterialIDs == ETargetMaterialIDs::OddIDs ? EmptyMaterialIDs : AutoUVSettings->MaterialIDs;
 
-	int32 UVLayer = AutoUVSettings->GetSelectedChannelIndex();
-
-	UE::PlanarCut::EUseMaterials UseMaterialIDs = GetUseMaterials(AutoUVSettings->TargetMaterialIDs);
+	BakeOp->UVLayer = AutoUVSettings->GetSelectedChannelIndex();
+	BakeOp->GutterSize = AutoUVSettings->GutterSize;
+	BakeOp->UseMaterialIDs = GetUseMaterials(AutoUVSettings->TargetMaterialIDs);
 
+	int32 OutputRes = (int32)AutoUVSettings->Resolution;
 	FImageDimensions Dimensions(OutputRes, OutputRes);
-
-	UE::Geometry::TImageBuilder<FVector4f> ImageBuilder;
-	ImageBuilder.SetDimensions(Dimensions);
-	ImageBuilder.Clear(FVector4f(0, 0, 0, 0));
+	BakeOp->ImageBuilder = MakeUnique<FImageType>();
+	BakeOp->ImageBuilder->SetDimensions(Dimensions);
+	BakeOp->ImageBuilder->Clear(FVector4f(0, 0, 0, 0));
 
 	typedef UE::PlanarCut::EBakeAttributes EBakeAttributes;
-	FIndex4i Attributes;
 	// Note: Ordering of these attributes should match the order and comments in the AutoUVSettings struct
 	//		 Update the order and comments there if you change the ordering here.
 	if (AutoUVSettings->BakeTextureType == ETextureType::ThicknessAndSurfaceAttributes)
 	{
-		Attributes = FIndex4i(
+		BakeOp->Attributes = FIndex4i(
 			AutoUVSettings->bDistToOuter ? (int32)EBakeAttributes::DistanceToExternal : 0,
 			AutoUVSettings->bAmbientOcclusion ? (int32)EBakeAttributes::AmbientOcclusion : 0,
 			AutoUVSettings->bSmoothedCurvature ? (int32)EBakeAttributes::Curvature : 0,
@@ -475,7 +545,7 @@ void UFractureToolAutoUV::BakeTextureForComponent(UGeometryCollectionComponent*
 	}
 	else if (AutoUVSettings->BakeTextureType == ETextureType::SpatialGradients)
 	{
-		Attributes = FIndex4i(
+		BakeOp->Attributes = FIndex4i(
 			(int32)EBakeAttributes::PositionX,
 			(int32)EBakeAttributes::PositionY,
 			(int32)EBakeAttributes::PositionZ,
@@ -484,30 +554,34 @@ void UFractureToolAutoUV::BakeTextureForComponent(UGeometryCollectionComponent*
 	}
 	else // ETextureType::Normals
 	{
-		Attributes = FIndex4i(
+		BakeOp->Attributes = FIndex4i(
 			(int32)EBakeAttributes::NormalX,
 			(int32)EBakeAttributes::NormalY,
 			(int32)EBakeAttributes::NormalZ,
 			0
 		);
 	}
-	UE::PlanarCut::FTextureAttributeSettings AttribSettings;
-	AttribSettings.ToExternal_MaxDistance = AutoUVSettings->MaxDistance;
-	AttribSettings.AO_Rays = AutoUVSettings->OcclusionRays;
-	AttribSettings.AO_BlurRadius = AutoUVSettings->OcclusionBlurRadius;
-	AttribSettings.Curvature_BlurRadius = AutoUVSettings->CurvatureBlurRadius;
-	AttribSettings.Curvature_SmoothingSteps = AutoUVSettings->SmoothingIterations;
-	AttribSettings.Curvature_VoxelRes = AutoUVSettings->VoxelResolution;
-	AttribSettings.Curvature_ThicknessFactor = AutoUVSettings->ThicknessFactor;
-	AttribSettings.Curvature_MaxValue = AutoUVSettings->MaxCurvature;
-	AttribSettings.ClearGutterChannel = 3; // default clear the gutters for the alpha channel, so it shows more clearly the island boundaries
-	UE::PlanarCut::TextureInternalSurfaces(UVLayer, Collection, AutoUVSettings->GutterSize, Attributes, AttribSettings, ImageBuilder,
-		UseMaterialIDs,
-		AutoUVSettings->TargetMaterialIDs == ETargetMaterialIDs::OddIDs ? EmptyMaterialIDs : AutoUVSettings->MaterialIDs);
-
-	if (Progress)
+	BakeOp->AttribSettings.ToExternal_MaxDistance = AutoUVSettings->MaxDistance;
+	BakeOp->AttribSettings.AO_Rays = AutoUVSettings->OcclusionRays;
+	BakeOp->AttribSettings.AO_BlurRadius = AutoUVSettings->OcclusionBlurRadius;
+	BakeOp->AttribSettings.Curvature_BlurRadius = AutoUVSettings->CurvatureBlurRadius;
+	BakeOp->AttribSettings.Curvature_SmoothingSteps = AutoUVSettings->SmoothingIterations;
+	BakeOp->AttribSettings.Curvature_VoxelRes = AutoUVSettings->VoxelResolution;
+	BakeOp->AttribSettings.Curvature_ThicknessFactor = AutoUVSettings->ThicknessFactor;
+	BakeOp->AttribSettings.Curvature_MaxValue = AutoUVSettings->MaxCurvature;
+	BakeOp->AttribSettings.ClearGutterChannel = 3; // default clear the gutters for the alpha channel, so it shows more clearly the island boundaries
+
+	FImageType ImageResult;
+	int Result = RunCancellableGeometryCollectionOpGeneric<FBakeTextureOp, FImageType>(
+		ImageResult, [](FImageType& ToUpdate, FImageType& Result)
+		{
+			ToUpdate = MoveTemp(Result);
+		},
+		MoveTemp(BakeOp), LOCTEXT("ComputingBakeTextureMessage", "Baking Textures"));
+	if (Result == -1)
 	{
-		Progress(1, LOCTEXT("SavingTexture", "Saving result"));
+		// Bake failed or was cancelled
+		return;
 	}
 
 	// choose default texture name based on corresponding geometry collection name
@@ -522,28 +596,28 @@ void UFractureToolAutoUV::BakeTextureForComponent(UGeometryCollectionComponent*
 		Suffix = "_AutoUV_Normals";
 	}
 	Suffix = FPaths::MakeValidFileName(Suffix);
-	SaveGeneratedTexture(ImageBuilder, FString::Printf(TEXT("%s%s"), *BaseName, *Suffix), Component->GetRestCollection(), AutoUVSettings->bPromptToSave, AutoUVSettings->bReplaceExisting);
+	SaveGeneratedTexture(ImageResult, FString::Printf(TEXT("%s%s"), *BaseName, *Suffix), Component->GetRestCollection(), AutoUVSettings->bPromptToSave, AutoUVSettings->bReplaceExisting);
 }
 
+
+
 int32 UFractureToolAutoUV::ExecuteFracture(const FFractureToolContext& FractureContext)
 {
 	if (FractureContext.GetGeometryCollection().IsValid())
 	{
 		bool bDoUVLayout = true;
 
-		FScopedSlowTask UVTask(bDoUVLayout ? 3 : 2, LOCTEXT("StartingAutoUV", "Automatically laying out and texturing internal surfaces"));
-		UVTask.MakeDialog();
-
 		UGeometryCollectionComponent* Component = FractureContext.GetGeometryCollectionComponent();
 
 		if (bDoUVLayout)
 		{
-			UVTask.EnterProgressFrame(1, LOCTEXT("LayOutUVIslands", "Laying out UV islands"));
-			LayoutUVsForComponent(Component);
+			if (!LayoutUVsForComponent(Component))
+			{
+				return INDEX_NONE; // Early out if layout failed or was cancelled
+			}
 		}
 
-		UVTask.EnterProgressFrame(1, LOCTEXT("TexturingSurfaces", "Texturing internal surfaces"));
-		BakeTextureForComponent(Component, [&](int32 AmtWork, const FText& Msg) { UVTask.EnterProgressFrame(AmtWork, Msg); });
+		BakeTextureForComponent(Component);
 	}
 
 	return INDEX_NONE;
diff --git a/Engine/Plugins/Experimental/ChaosEditor/Source/FractureEditor/Private/FractureToolUV.h b/Engine/Plugins/Experimental/ChaosEditor/Source/FractureEditor/Private/FractureToolUV.h
--- a/Engine/Plugins/Experimental/ChaosEditor/Source/FractureEditor/Private/FractureToolUV.h
+++ b/Engine/Plugins/Experimental/ChaosEditor/Source/FractureEditor/Private/FractureToolUV.h
@@ -254,5 +254,5 @@ protected:
 	bool SaveGeneratedTexture(UE::Geometry::TImageBuilder<FVector4f>& ImageBuilder, FString ObjectBaseName, const UObject* RelativeToAsset, bool bPromptToSave, bool bAllowReplace);
 
 	bool LayoutUVsForComponent(UGeometryCollectionComponent* Component);
-	void BakeTextureForComponent(UGeometryCollectionComponent* Component, TFunction<void(int32, const FText&)> Progress = nullptr);
+	void BakeTextureForComponent(UGeometryCollectionComponent* Component);
 };
diff --git a/Engine/Plugins/Experimental/PlanarCutPlugin/Source/PlanarCut/Private/FractureAutoUV.cpp b/Engine/Plugins/Experimental/PlanarCutPlugin/Source/PlanarCut/Private/FractureAutoUV.cpp
--- a/Engine/Plugins/Experimental/PlanarCutPlugin/Source/PlanarCut/Private/FractureAutoUV.cpp
+++ b/Engine/Plugins/Experimental/PlanarCutPlugin/Source/PlanarCut/Private/FractureAutoUV.cpp
@@ -26,6 +26,7 @@
 #include "Image/ImageOccupancyMap.h"
 #include "VectorUtil.h"
 #include "FrameTypes.h"
+#include "Util/ProgressCancel.h"
 
 #include "Templates/PimplPtr.h"
 
@@ -136,7 +137,7 @@ namespace {
  * Shared helper method to set 'active' subset of triangles for a collection, to be considered for texturing
  * @return number of active triangles
  */
-int32 SetActiveTriangles(FGeometryCollection* Collection, TArray<bool>& ActiveTrianglesOut,
+int32 SetActiveTriangles(const FGeometryCollection* Collection, TArray<bool>& ActiveTrianglesOut,
 	bool bActivateInsideTriangles, EUseMaterials MaterialsPattern, TArrayView<int32> WhichMaterialsAreInside)
 {
 	TSet<int32> InsideMaterials;
@@ -446,9 +447,12 @@ bool UVLayout(
 	float GutterSize,
 	EUseMaterials MaterialsPattern,
 	TArrayView<int32> WhichMaterials,
-	bool bRecreateUVsForDegenerateIslands
+	bool bRecreateUVsForDegenerateIslands,
+	FProgressCancel* Progress
 )
 {
+	FProgressCancel::FProgressScope InitScope = FProgressCancel::CreateScopeTo(Progress, .2);
+
 	TArray<bool> ActiveTriangles;
 	int32 NumActive = SetActiveTriangles(&Collection, ActiveTriangles, true, MaterialsPattern, WhichMaterials);
 	FGeomMesh UVMesh(TargetUVLayer, &Collection, ActiveTriangles, NumActive);
@@ -456,8 +460,17 @@ bool UVLayout(
 	TArray<TArray<int32>> UVIslands;
 	UE::UVPacking::CreateUVIslandsFromMeshTopology<FGeomMesh>(UVMesh, UVIslands);
 
+	InitScope.Done();
+
+	if (Progress && Progress->Cancelled())
+	{
+		return false;
+	}
+
 	if (bRecreateUVsForDegenerateIslands)
 	{
+		FProgressCancel::FProgressScope ForDegenerateIslandsScope = FProgressCancel::CreateScopeTo(Progress, .5);
+
 		TArray<int> IslandVert; IslandVert.SetNumUninitialized(UVMesh.MaxVertexID());
 
 		const int32 NumIslands = UVIslands.Num();
@@ -556,6 +569,12 @@ bool UVLayout(
 		}
 	}
 
+	if (Progress && Progress->Cancelled())
+	{
+		return false;
+	}
+	FProgressCancel::FProgressScope FinalScope = FProgressCancel::CreateScopeTo(Progress, 1);
+
 	UE::Geometry::FUVPacker Packer;
 	Packer.bScaleIslandsByWorldSpaceTexelRatio = true; // let packer scale islands separately to have consistent texel-to-world ratio
 	Packer.bAllowFlips = false;
@@ -568,7 +587,7 @@ bool UVLayout(
 }
 
 
-void TextureInternalSurfaces(
+bool TextureInternalSurfaces(
 	int32 TargetUVLayer,
 	FGeometryCollection& Collection,
 	int32 GutterSize,
@@ -576,7 +595,8 @@ void TextureInternalSurfaces(
 	const FTextureAttributeSettings& AttributeSettings,
 	TImageBuilder<FVector4f>& TextureOut,
 	EUseMaterials MaterialsPattern,
-	TArrayView<int32> WhichMaterials
+	TArrayView<int32> WhichMaterials,
+	FProgressCancel* Progress
 )
 {
 	TArray<bool> ToTextureTriangles;
@@ -615,9 +635,21 @@ void TextureInternalSurfaces(
 	{
 		for (int32 GeomIdx = 0; GeomIdx < Collection.TransformIndex.Num(); GeomIdx++)
 		{
-			TransformIndices.Add(Collection.TransformIndex[GeomIdx]);
+			int32 TransformIdx = Collection.TransformIndex[GeomIdx];
+			if (Collection.IsVisible(TransformIdx))
+			{
+				TransformIndices.Add(TransformIdx);
+			}
 		}
 	}
+
+	const float AmbientWorkPer = 1, CurvatureWorkPer = 10;
+	const float NumMeshes = TransformIndices.Num();
+	const float AmountOfWork = 
+		float(AmbientIdx > -1) * AmbientWorkPer * NumMeshes +
+		float(CurvatureIdx > -1) * CurvatureWorkPer * NumMeshes;
+	const float WorkIncr = AmountOfWork > 0 ? .9 / AmountOfWork : 0;
+
 	FDynamicMeshCollection CollectionMeshes(&Collection, TransformIndices, FTransform::Identity, false);
 	if (bNeedsDynamicMeshes)
 	{
@@ -648,6 +680,11 @@ void TextureInternalSurfaces(
 		}
 	}
 
+	if (Progress && Progress->Cancelled())
+	{
+		return false;
+	}
+
 	TArray64<float> CurvatureValues; // buffer to fill with curvature values
 	if (CurvatureIdx > -1)
 	{
@@ -657,6 +694,11 @@ void TextureInternalSurfaces(
 		CurvatureValues.SetNumZeroed(TexelsNum);
 		for (int MeshIdx = 0; MeshIdx < CollectionMeshes.Meshes.Num(); MeshIdx++)
 		{
+			if (Progress && Progress->Cancelled())
+			{
+				return false;
+			}
+			FProgressCancel::FProgressScope CurvatureScope(Progress, CurvatureWorkPer * WorkIncr);
 
 			int32 TransformIdx = CollectionMeshes.Meshes[MeshIdx].TransformIndex;
 			FDynamicMesh3& Mesh = CollectionMeshes.Meshes[MeshIdx].AugMesh;
@@ -730,6 +772,12 @@ void TextureInternalSurfaces(
 		AmbientValues.Init(1.0f, TexelsNum);
 		for (int MeshIdx = 0; MeshIdx < CollectionMeshes.Meshes.Num(); MeshIdx++)
 		{
+			if (Progress && Progress->Cancelled())
+			{
+				return false;
+			}
+			FProgressCancel::FProgressScope CurvatureScope(Progress, AmbientWorkPer * WorkIncr);
+
 			int32 TransformIdx = CollectionMeshes.Meshes[MeshIdx].TransformIndex;
 			FDynamicMesh3& Mesh = CollectionMeshes.Meshes[MeshIdx].AugMesh;
 
@@ -767,6 +815,12 @@ void TextureInternalSurfaces(
 		}
 	}
 
+	if (Progress && Progress->Cancelled())
+	{
+		return false;
+	}
+	FProgressCancel::FProgressScope EndBake = FProgressCancel::CreateScopeTo(Progress, 1);
+
 	ParallelFor(OccupancyMap.Dimensions.GetHeight(),
 		[&AttributeSettings, &TextureOut, &UVMesh, &OccupancyMap, &ToTextureMesh, &OutsideSpatial,
 		 &DistanceToExternalIdx, &AmbientIdx, &NormalXIdx, &NormalYIdx, &NormalZIdx, &PosXIdx, &PosYIdx, &PosZIdx](int32 Y)
@@ -839,6 +893,11 @@ void TextureInternalSurfaces(
 		}
 	});
 
+	if (Progress && Progress->Cancelled())
+	{
+		return false;
+	}
+
 	auto CopyValuesToChannel = [&TextureOut, &OccupancyMap](int ChannelIdx, TArray64<float>& Values)
 	{
 		if (ChannelIdx > -1)
@@ -872,6 +931,8 @@ void TextureInternalSurfaces(
 			TextureOut.CopyPixel(GutterToInside.Value, GutterToInside.Key);
 		}
 	}
+
+	return true;
 }
 
 }} // namespace UE::PlanarCut
\ No newline at end of file
diff --git a/Engine/Plugins/Experimental/PlanarCutPlugin/Source/PlanarCut/Public/FractureAutoUV.h b/Engine/Plugins/Experimental/PlanarCutPlugin/Source/PlanarCut/Public/FractureAutoUV.h
--- a/Engine/Plugins/Experimental/PlanarCutPlugin/Source/PlanarCut/Public/FractureAutoUV.h
+++ b/Engine/Plugins/Experimental/PlanarCutPlugin/Source/PlanarCut/Public/FractureAutoUV.h
@@ -6,6 +6,8 @@
 
 #include "Image/ImageBuilder.h"
 
+class FProgressCancel;
+
 namespace UE { namespace PlanarCut {
 
 enum class EUseMaterials
@@ -50,7 +52,8 @@ bool PLANARCUT_API UVLayout(
 	float GutterSize = 1,
 	EUseMaterials MaterialsPattern = EUseMaterials::OddMaterials,
 	TArrayView<int32> WhichMaterials = TArrayView<int32>(),
-	bool bRecreateUVsForDegenerateIslands = true
+	bool bRecreateUVsForDegenerateIslands = true,
+	FProgressCancel* Progress = nullptr
 );
 
 
@@ -102,7 +105,7 @@ struct FTextureAttributeSettings
  * @param MaterialsPattern	Which pattern of material IDs to apply texture to
  * @param WhichMaterials	If non-empty, apply texture to the listed material IDs
  */
-void PLANARCUT_API TextureInternalSurfaces(
+bool PLANARCUT_API TextureInternalSurfaces(
 	int32 TargetUVLayer,
 	FGeometryCollection& Collection,
 	int32 GutterSize,
@@ -110,7 +113,8 @@ void PLANARCUT_API TextureInternalSurfaces(
 	const FTextureAttributeSettings& AttributeSettings,
 	UE::Geometry::TImageBuilder<FVector4f>& TextureOut,
 	EUseMaterials MaterialsPattern = EUseMaterials::OddMaterials,
-	TArrayView<int32> WhichMaterials = TArrayView<int32>()
+	TArrayView<int32> WhichMaterials = TArrayView<int32>(),
+	FProgressCancel* Progress = nullptr
 );
 
 }} // namespace UE::PlanarCut
\ No newline at end of file
