Fracture Tool : make outliner properly display cluster crumbling option for the remove on break feature

diff --git a/Engine/Plugins/Experimental/ChaosEditor/Source/FractureEditor/Private/SGeometryCollectionOutliner.cpp b/Engine/Plugins/Experimental/ChaosEditor/Source/FractureEditor/Private/SGeometryCollectionOutliner.cpp
--- a/Engine/Plugins/Experimental/ChaosEditor/Source/FractureEditor/Private/SGeometryCollectionOutliner.cpp
+++ b/Engine/Plugins/Experimental/ChaosEditor/Source/FractureEditor/Private/SGeometryCollectionOutliner.cpp
@@ -196,7 +196,7 @@ TSharedRef<SWidget> SGeometryCollectionOutlinerRow::GenerateWidgetForColumn(cons
 	if (ColumnName == SGeometryCollectionOutlinerColumnID::Broken)
 		return Item->MakeBrokenColumnWidget();
 	if (ColumnName == SGeometryCollectionOutlinerColumnID::PostBreakTime)
-		return Item->MakePostBreakTimecolumnWidget();
+		return Item->MakePostBreakTimeColumnWidget();
 	if (ColumnName == SGeometryCollectionOutlinerColumnID::RemovalTime)
 		return Item->MakeRemovalTimeColumnWidget();
 	if (ColumnName == SGeometryCollectionOutlinerColumnID::ImportedCollisions)
@@ -953,14 +953,14 @@ static FText FormatRemoveOnBreakTimeData(bool bDataAvailable, bool bEnabled, con
 	{
 		if (bEnabled)
 		{
-			return FText::Format(LOCTEXT("RemoveOnBreakFormat", "[{0}s, {1}s]"),Timer.X, Timer.Y);
+			return FText::Format(LOCTEXT("GCOutliner_RemoveOnBreakTimer_Format", "[{0}s, {1}s]"),Timer.X, Timer.Y);
 		}
-		return FText(LOCTEXT("RemoveOnBreakEmpty", "[-, -]"));
+		return FText(LOCTEXT("GCOutliner_RemoveOnBreakTimer_Empty", "[-, -]"));
 	}
 	return FText();
 }
 
-TSharedRef<SWidget> FGeometryCollectionTreeItemBone::MakePostBreakTimecolumnWidget() const
+TSharedRef<SWidget> FGeometryCollectionTreeItemBone::MakePostBreakTimeColumnWidget() const
 {
 	return SNew(SHorizontalBox)
 		+ SHorizontalBox::Slot()
@@ -975,14 +975,24 @@ TSharedRef<SWidget> FGeometryCollectionTreeItemBone::MakePostBreakTimecolumnWidg
 
 TSharedRef<SWidget> FGeometryCollectionTreeItemBone::MakeRemovalTimeColumnWidget() const
 {
-	const bool EnableRemovalTimer = RemoveOnBreak.IsEnabled() && !(IsCluster && RemoveOnBreak.GetClusterCrumbling());  
+	const bool EnableRemovalTimer = RemoveOnBreak.IsEnabled() && !(IsCluster && RemoveOnBreak.GetClusterCrumbling());
+	const bool ClusterCrumbling = RemoveOnBreak.IsEnabled() && (IsCluster && RemoveOnBreak.GetClusterCrumbling());
+	FText RemovalTimeText; 
+	if (ClusterCrumbling)
+	{
+		RemovalTimeText = FText(LOCTEXT("GCOutliner_ClusterCrumbling_Text", "Cluster Crumbling"));
+	}
+	else
+	{
+		RemovalTimeText = FormatRemoveOnBreakTimeData(RemoveOnBreakAvailable, EnableRemovalTimer,  RemoveOnBreak.GetRemovalTimer());
+	}
 	return SNew(SHorizontalBox)
 	+ SHorizontalBox::Slot()
 		.Padding(12.f, 0.f)
 		.HAlign(HAlign_Center)
 		[
 			SNew(STextBlock)
-			.Text(FormatRemoveOnBreakTimeData(RemoveOnBreakAvailable, EnableRemovalTimer,  RemoveOnBreak.GetRemovalTimer()))
+			.Text(RemovalTimeText)
 			.ColorAndOpacity(ItemColor)
 		];
 }
diff --git a/Engine/Plugins/Experimental/ChaosEditor/Source/FractureEditor/Private/SGeometryCollectionOutliner.h b/Engine/Plugins/Experimental/ChaosEditor/Source/FractureEditor/Private/SGeometryCollectionOutliner.h
--- a/Engine/Plugins/Experimental/ChaosEditor/Source/FractureEditor/Private/SGeometryCollectionOutliner.h
+++ b/Engine/Plugins/Experimental/ChaosEditor/Source/FractureEditor/Private/SGeometryCollectionOutliner.h
@@ -164,7 +164,7 @@ public:
 	TSharedRef<SWidget> MakeDamageThresholdColumnWidget() const;
 	TSharedRef<SWidget> MakeBrokenColumnWidget() const;
 	TSharedRef<SWidget> MakeInitialStateColumnWidget() const;
-	TSharedRef<SWidget> MakePostBreakTimecolumnWidget() const;
+	TSharedRef<SWidget> MakePostBreakTimeColumnWidget() const;
 	TSharedRef<SWidget> MakeRemovalTimeColumnWidget() const;
 	TSharedRef<SWidget> MakeImportedCollisionsColumnWidget() const;
 	TSharedRef<SWidget> MakeEmptyColumnWidget() const;
