Fracture editor : fix crash when compiling geometry collection blueprint  while in fracture editor ( 3rd iteration )

diff --git a/Engine/Plugins/Experimental/ChaosEditor/Source/FractureEditor/Private/SGeometryCollectionOutliner.cpp b/Engine/Plugins/Experimental/ChaosEditor/Source/FractureEditor/Private/SGeometryCollectionOutliner.cpp
--- a/Engine/Plugins/Experimental/ChaosEditor/Source/FractureEditor/Private/SGeometryCollectionOutliner.cpp
+++ b/Engine/Plugins/Experimental/ChaosEditor/Source/FractureEditor/Private/SGeometryCollectionOutliner.cpp
@@ -537,8 +537,8 @@ void SGeometryCollectionOutliner::SetInitialDynamicState(int32 InDynamicState)
 
 TSharedRef<ITableRow> FGeometryCollectionTreeItemComponent::MakeTreeRowWidget(const TSharedRef<STableViewBase>& InOwnerTable, bool bNoExtraColumn)
 {
-	FString ActorName = Component->GetOwner()->GetActorLabel();
-	FString ComponentName = Component->GetClass()->GetFName().ToString();
+	FString ActorName = Component.IsValid()? Component->GetOwner()->GetActorLabel(): FString();
+	FString ComponentName = Component.IsValid()? Component->GetClass()->GetFName().ToString() : FString();
 
 	return SNew(STableRow<FGeometryCollectionTreeItemPtr>, InOwnerTable)
 		.Content()
@@ -713,36 +713,38 @@ void FGeometryCollectionTreeItemComponent::SetHistogramSelection(TArray<int32>&
 
 bool FGeometryCollectionTreeItemComponent::FilterBoneIndex(int32 BoneIndex) const
 {
-	const FGeometryCollection* Collection = Component->GetRestCollection()->GetGeometryCollection().Get();
-	const TManagedArray<int32>& SimTypes = Collection->SimulationType;
-	bool bHasChildren = Collection->Children[BoneIndex].Num() > 0;
-
-	if (SimTypes[BoneIndex] != FGeometryCollection::ESimulationTypes::FST_Clustered)
+	if (Component.IsValid())
 	{
-		// We only display cluster nodes deeper than the view level.
-		UFractureSettings* FractureSettings = GetMutableDefault<UFractureSettings>();
+		const FGeometryCollection* Collection = Component->GetRestCollection()->GetGeometryCollection().Get();
+		const TManagedArray<int32>& SimTypes = Collection->SimulationType;
+		bool bHasChildren = Collection->Children[BoneIndex].Num() > 0;
 
-		if (FractureSettings->FractureLevel >= 0)
+		if (SimTypes[BoneIndex] != FGeometryCollection::ESimulationTypes::FST_Clustered)
 		{
-			const TManagedArray<int32>& Level = Collection->GetAttribute<int32>("Level", FTransformCollection::TransformGroup);
-			int32 BoneLevel = Level[BoneIndex];
-			// bone is not at the right level itself and doesn't have child(ren) at the right level
-			if (BoneLevel != FractureSettings->FractureLevel && (!bHasChildren || BoneLevel + 1 != FractureSettings->FractureLevel))
+			// We only display cluster nodes deeper than the view level.
+			UFractureSettings* FractureSettings = GetMutableDefault<UFractureSettings>();
+
+			if (FractureSettings->FractureLevel >= 0)
 			{
-				return false;
+				const TManagedArray<int32>& Level = Collection->GetAttribute<int32>("Level", FTransformCollection::TransformGroup);
+				int32 BoneLevel = Level[BoneIndex];
+				// bone is not at the right level itself and doesn't have child(ren) at the right level
+				if (BoneLevel != FractureSettings->FractureLevel && (!bHasChildren || BoneLevel + 1 != FractureSettings->FractureLevel))
+				{
+					return false;
+				}
 			}
-		}
 
-		// If anything is selected int the Histogram, we filter by that selection.
-		if (HistogramSelection.Num() > 0)
-		{
-			if (!HistogramSelection.Contains(BoneIndex))
+			// If anything is selected int the Histogram, we filter by that selection.
+			if (HistogramSelection.Num() > 0)
 			{
-				return false;
-			}
-		}		
+				if (!HistogramSelection.Contains(BoneIndex))
+				{
+					return false;
+				}
+			}		
+		}
 	}
-
 	return true;	
 }
 
@@ -773,82 +775,85 @@ void FGeometryCollectionTreeItemBone::UpdateItemFromCollection()
 
 	// Set color according to simulation type
 	const UGeometryCollectionComponent* Component = ParentComponentItem->GetComponent();
-	const FDamageCollector* Collector = Component->GetRunTimeDataCollector();
-	const UGeometryCollection* RestCollection = Component->GetRestCollection();
-	if (RestCollection && IsValidChecked(RestCollection))
+	if (Component)
 	{
-		TSharedPtr<FGeometryCollection, ESPMode::ThreadSafe> GeometryCollectionPtr = RestCollection->GetGeometryCollection();
-		const TManagedArray<int32>& SimulationType = GeometryCollectionPtr->SimulationType;
-		const TManagedArray<float>* RelativeSizes = GeometryCollectionPtr->FindAttribute<float>("Size", FTransformCollection::TransformGroup);
-		const TManagedArray<FVector4f>* RemoveOnBreakArray = GeometryCollectionPtr->FindAttribute<FVector4f>("RemoveOnBreak", FTransformCollection::TransformGroup);
+		const FDamageCollector* Collector = Component->GetRunTimeDataCollector();
+		const UGeometryCollection* RestCollection = Component->GetRestCollection();
+		if (RestCollection && IsValidChecked(RestCollection))
+		{
+			TSharedPtr<FGeometryCollection, ESPMode::ThreadSafe> GeometryCollectionPtr = RestCollection->GetGeometryCollection();
+			const TManagedArray<int32>& SimulationType = GeometryCollectionPtr->SimulationType;
+			const TManagedArray<float>* RelativeSizes = GeometryCollectionPtr->FindAttribute<float>("Size", FTransformCollection::TransformGroup);
+			const TManagedArray<FVector4f>* RemoveOnBreakArray = GeometryCollectionPtr->FindAttribute<FVector4f>("RemoveOnBreak", FTransformCollection::TransformGroup);
 
-		using FImplicitGeom = FGeometryDynamicCollection::FSharedImplicit;
-		const TManagedArray<FImplicitGeom>* ExternalCollisions = GeometryCollectionPtr->FindAttribute<FImplicitGeom>("ExternalCollisions", FTransformCollection::TransformGroup);
+			using FImplicitGeom = FGeometryDynamicCollection::FSharedImplicit;
+			const TManagedArray<FImplicitGeom>* ExternalCollisions = GeometryCollectionPtr->FindAttribute<FImplicitGeom>("ExternalCollisions", FTransformCollection::TransformGroup);
 		
-		if (ensure(ItemBoneIndex >= 0 && ItemBoneIndex < SimulationType.Num()))
-		{
-			const bool bHasLevelAttribute = GeometryCollectionPtr->HasAttribute("Level", FGeometryCollection::TransformGroup);
-			if (bHasLevelAttribute && OutlinerSettings->ColorByLevel)
-			{
-				const TManagedArray<int32>& Level = GeometryCollectionPtr->GetAttribute<int32>("Level", FTransformCollection::TransformGroup);
-				ItemColor = FSlateColor(FGeometryCollectionTreeItem::GetColorPerDepth((uint32)Level[ItemBoneIndex]));
-			}
-			else
+			if (ensure(ItemBoneIndex >= 0 && ItemBoneIndex < SimulationType.Num()))
 			{
-				switch (SimulationType[ItemBoneIndex])
+				const bool bHasLevelAttribute = GeometryCollectionPtr->HasAttribute("Level", FGeometryCollection::TransformGroup);
+				if (bHasLevelAttribute && OutlinerSettings->ColorByLevel)
 				{
-				case FGeometryCollection::ESimulationTypes::FST_None:
-					ItemColor = FLinearColor::Green;
-					break;
-
-				case FGeometryCollection::ESimulationTypes::FST_Rigid:
-					if (GeometryCollectionPtr->IsVisible(ItemBoneIndex))
-					{
-						ItemColor = FSlateColor::UseForeground();
-					}
-					else
+					const TManagedArray<int32>& Level = GeometryCollectionPtr->GetAttribute<int32>("Level", FTransformCollection::TransformGroup);
+					ItemColor = FSlateColor(FGeometryCollectionTreeItem::GetColorPerDepth((uint32)Level[ItemBoneIndex]));
+				}
+				else
+				{
+					switch (SimulationType[ItemBoneIndex])
 					{
-						ItemColor = InvalidColor;
+					case FGeometryCollection::ESimulationTypes::FST_None:
+						ItemColor = FLinearColor::Green;
+						break;
+
+					case FGeometryCollection::ESimulationTypes::FST_Rigid:
+						if (GeometryCollectionPtr->IsVisible(ItemBoneIndex))
+						{
+							ItemColor = FSlateColor::UseForeground();
+						}
+						else
+						{
+							ItemColor = InvalidColor;
+						}
+						break;
+
+					case FGeometryCollection::ESimulationTypes::FST_Clustered:
+						ItemColor = FSlateColor(FColor::Cyan);
+						break;
+
+					default:
+						ensureMsgf(false, TEXT("Invalid Geometry Collection simulation type encountered."));
+						break;
 					}
-					break;
-
-				case FGeometryCollection::ESimulationTypes::FST_Clustered:
-					ItemColor = FSlateColor(FColor::Cyan);
-					break;
+				}
 
-				default:
-					ensureMsgf(false, TEXT("Invalid Geometry Collection simulation type encountered."));
-					break;
+				InitialState = GeometryCollectionPtr->InitialDynamicState[ItemBoneIndex];
+				if (RelativeSizes)
+				{
+					RelativeSize = (*RelativeSizes)[ItemBoneIndex];
 				}
-			}
 
-			InitialState = GeometryCollectionPtr->InitialDynamicState[ItemBoneIndex];
-			if (RelativeSizes)
-			{
-				RelativeSize = (*RelativeSizes)[ItemBoneIndex];
-			}
+				if (Collector)
+				{
+					const FDamageCollector::FDamageData& DamageData = (*Collector)[ItemBoneIndex];
+					Damage = DamageData.MaxDamages;
+					DamageThreshold = DamageData.DamageThreshold;
+					Broken = DamageData.bIsBroken;
+				}
 
-			if (Collector)
-			{
-				const FDamageCollector::FDamageData& DamageData = (*Collector)[ItemBoneIndex];
-				Damage = DamageData.MaxDamages;
-				DamageThreshold = DamageData.DamageThreshold;
-				Broken = DamageData.bIsBroken;
-			}
+				RemoveOnBreakAvailable = (RemoveOnBreakArray != nullptr); 
+				if (RemoveOnBreakAvailable)
+				{
+					RemoveOnBreak = (*RemoveOnBreakArray)[ItemBoneIndex];
+				}
 
-			RemoveOnBreakAvailable = (RemoveOnBreakArray != nullptr); 
-			if (RemoveOnBreakAvailable)
-			{
-				RemoveOnBreak = (*RemoveOnBreakArray)[ItemBoneIndex];
-			}
+				ImportedCollisionsUsed = RestCollection->bImportCollisionFromSource;
+				if (ExternalCollisions)
+				{
+					ImportedCollisionsAvailable = (*ExternalCollisions)[ItemBoneIndex] != nullptr;
+				}
 
-			ImportedCollisionsUsed = RestCollection->bImportCollisionFromSource;
-			if (ExternalCollisions)
-			{
-				ImportedCollisionsAvailable = (*ExternalCollisions)[ItemBoneIndex] != nullptr;
+				IsCluster = (GeometryCollectionPtr->Children[ItemBoneIndex].Num() > 0);
 			}
-
-			IsCluster = (GeometryCollectionPtr->Children[ItemBoneIndex].Num() > 0);
 		}
 	}
 }
