Fracture editor : Set Remove on break tool  : disable removal trimer when cluster crumbing is on

diff --git a/Engine/Plugins/Experimental/ChaosEditor/Source/FractureEditor/Private/FractureToolProperties.h b/Engine/Plugins/Experimental/ChaosEditor/Source/FractureEditor/Private/FractureToolProperties.h
--- a/Engine/Plugins/Experimental/ChaosEditor/Source/FractureEditor/Private/FractureToolProperties.h
+++ b/Engine/Plugins/Experimental/ChaosEditor/Source/FractureEditor/Private/FractureToolProperties.h
@@ -72,8 +72,8 @@ public:
 		: Super(ObjInit)
 		, Enabled(true)
 		, PostBreakTimer(0,0)
-		, RemovalTimer(2,3)
 		, ClusterCrumbling(false)
+		, RemovalTimer(2,3)
 	{}
 
 	/** whether or not the remove on fracture is enabled */
@@ -84,14 +84,14 @@ public:
 	UPROPERTY(EditAnywhere, Category = SetRemoveOnBreak, meta = (DisplayName = "Post Break Timer", EditCondition="Enabled"))
 	FVector2f PostBreakTimer;  
 
-	/** Min/Max time for how long removal lasts */
-	UPROPERTY(EditAnywhere, Category = SetRemoveOnBreak, meta = (DisplayName = "Removal Timer", EditCondition="Enabled"))
-	FVector2f RemovalTimer;
-	
 	/** When set, clusters will crumble when post break timer expires, non clusters will simply use the removal timer */
-	UPROPERTY(EditAnywhere, Category = SetRemoveOnBreak, meta = (DisplayName = "Cluster Crumbling"))
+	UPROPERTY(EditAnywhere, Category = SetRemoveOnBreak, meta = (DisplayName = "Cluster Crumbling", EditCondition="Enabled"))
 	bool ClusterCrumbling;
 
+	/** Min/Max time for how long removal lasts - not applicable when cluster crumbling is on  */
+	UPROPERTY(EditAnywhere, Category = SetRemoveOnBreak, meta = (DisplayName = "Removal Timer", EditCondition="!ClusterCrumbling && Enabled"))
+	FVector2f RemovalTimer;
+	
 	/** remove the remove on break attribute from the geometry collection, usefull to save memory on the asset if not needed */
 	UFUNCTION(CallInEditor, Category = Actions, meta = (DisplayName = "Delete Remove-On-Break Data"))
 	void DeleteRemoveOnBreakData();
