Add project-default convex generation settings to FractureMode project settings

Add options to save/load from them in the convex tool

Use project-default convex generation settings for new/reset geometry collections

diff --git a/Engine/Plugins/Experimental/ChaosEditor/Source/FractureEditor/Private/FractureModeSettings.cpp b/Engine/Plugins/Experimental/ChaosEditor/Source/FractureEditor/Private/FractureModeSettings.cpp
new file mode 100644
--- /dev/null
+++ b/Engine/Plugins/Experimental/ChaosEditor/Source/FractureEditor/Private/FractureModeSettings.cpp
@@ -0,0 +1,16 @@
+// Copyright Epic Games, Inc. All Rights Reserved.
+
+#include "FractureModeSettings.h"
+
+#include "GeometryCollection/GeometryCollection.h"
+#include "GeometryCollection/GeometryCollectionConvexPropertiesInterface.h"
+
+void UFractureModeSettings::ApplyDefaultConvexSettings(FGeometryCollection& GeometryCollection) const
+{
+	FGeometryCollectionConvexPropertiesInterface::FConvexCreationProperties Properties = GeometryCollection.GetConvexProperties();
+	Properties.FractionRemove = ConvexFractionAllowRemove;
+	Properties.SimplificationThreshold = ConvexSimplificationDistanceThreshold;
+	Properties.CanExceedFraction = ConvexCanExceedFraction;
+	Properties.bRemoveOverlaps = bConvexRemoveOverlaps;
+	GeometryCollection.SetConvexProperties(Properties);
+}
\ No newline at end of file
diff --git a/Engine/Plugins/Experimental/ChaosEditor/Source/FractureEditor/Private/FractureToolConvex.cpp b/Engine/Plugins/Experimental/ChaosEditor/Source/FractureEditor/Private/FractureToolConvex.cpp
--- a/Engine/Plugins/Experimental/ChaosEditor/Source/FractureEditor/Private/FractureToolConvex.cpp
+++ b/Engine/Plugins/Experimental/ChaosEditor/Source/FractureEditor/Private/FractureToolConvex.cpp
@@ -3,6 +3,7 @@
 #include "FractureToolConvex.h"
 
 #include "FractureToolContext.h"
+#include "FractureModeSettings.h"
 
 #include "GeometryCollection/GeometryCollectionObject.h"
 #include "GeometryCollection/GeometryCollectionComponent.h"
@@ -34,6 +35,28 @@ void UFractureConvexSettings::ClearCustomConvex()
 	ConvexTool->ClearCustomConvex();
 }
 
+void UFractureConvexSettings::SaveAsDefaults()
+{
+	UFractureModeSettings* ModeSettings = GetMutableDefault<UFractureModeSettings>();
+	ModeSettings->Modify();
+	ModeSettings->ConvexFractionAllowRemove = FractionAllowRemove;
+	ModeSettings->ConvexCanExceedFraction = CanExceedFraction;
+	ModeSettings->ConvexSimplificationDistanceThreshold = SimplificationDistanceThreshold;
+	ModeSettings->bConvexRemoveOverlaps = bRemoveOverlaps;
+}
+
+void UFractureConvexSettings::SetFromDefaults()
+{
+	const UFractureModeSettings* ModeSettings = GetDefault<UFractureModeSettings>();
+	FractionAllowRemove = ModeSettings->ConvexFractionAllowRemove;
+	CanExceedFraction = ModeSettings->ConvexCanExceedFraction;
+	SimplificationDistanceThreshold = ModeSettings->ConvexSimplificationDistanceThreshold;
+	bRemoveOverlaps = ModeSettings->bConvexRemoveOverlaps;
+
+	UFractureToolConvex* ConvexTool = Cast<UFractureToolConvex>(OwnerTool.Get());
+	ConvexTool->NotifyOfPropertyChangeByTool(this);
+}
+
 void UFractureToolConvex::DeleteConvexFromSelected()
 {
 	TArray<FFractureToolContext> FractureContexts = GetFractureToolContexts();
@@ -300,4 +323,9 @@ int32 UFractureToolConvex::ExecuteFracture(const FFractureToolContext& FractureC
 	return INDEX_NONE;
 }
 
+void UFractureToolConvex::Setup()
+{
+	Super::Setup();
+}
+
 #undef LOCTEXT_NAMESPACE
diff --git a/Engine/Plugins/Experimental/ChaosEditor/Source/FractureEditor/Private/FractureToolConvex.h b/Engine/Plugins/Experimental/ChaosEditor/Source/FractureEditor/Private/FractureToolConvex.h
--- a/Engine/Plugins/Experimental/ChaosEditor/Source/FractureEditor/Private/FractureToolConvex.h
+++ b/Engine/Plugins/Experimental/ChaosEditor/Source/FractureEditor/Private/FractureToolConvex.h
@@ -50,6 +50,14 @@ public:
 	/** Clear any manual adjustments to convex hulls on the selected bones */
 	UFUNCTION(CallInEditor, Category = Custom, meta = (DisplayName = "Clear Custom Convex"))
 	void ClearCustomConvex();
+
+	/** Save settings as project defaults, to be used for all new geometry collections */
+	UFUNCTION(CallInEditor, Category = Defaults, meta = (DisplayName = "Save As Defaults"))
+	void SaveAsDefaults();
+
+	/** Set settings from current project defaults */
+	UFUNCTION(CallInEditor, Category = Defaults, meta = (DisplayName = "Set From Defaults"))
+	void SetFromDefaults();
 };
 
 
@@ -89,6 +97,8 @@ public:
 	void PromoteChildren();
 	void ClearCustomConvex();
 
+	virtual void Setup() override;
+
 
 protected:
 	UPROPERTY(EditAnywhere, Category = Slicing)
diff --git a/Engine/Plugins/Experimental/ChaosEditor/Source/FractureEditor/Private/FractureToolGenerators.cpp b/Engine/Plugins/Experimental/ChaosEditor/Source/FractureEditor/Private/FractureToolGenerators.cpp
--- a/Engine/Plugins/Experimental/ChaosEditor/Source/FractureEditor/Private/FractureToolGenerators.cpp
+++ b/Engine/Plugins/Experimental/ChaosEditor/Source/FractureEditor/Private/FractureToolGenerators.cpp
@@ -13,6 +13,8 @@
 #include "Engine/StaticMesh.h"
 #include "Engine/StaticMeshActor.h"
 
+#include "FractureModeSettings.h"
+
 #include "GeometryCollection/GeometryCollectionComponent.h"
 #include "GeometryCollection/GeometryCollectionActor.h"
 #include "GeometryCollection/GeometryCollection.h"
@@ -264,6 +266,9 @@ AGeometryCollectionActor* UFractureToolGenerateAsset::ConvertActorsToGeometryCol
 	FGeometryCollectionProximityUtility ProximityUtility(FracturedGeometryCollection->GetGeometryCollection().Get());
 	ProximityUtility.UpdateProximity();
 
+	const UFractureModeSettings* ModeSettings = GetDefault<UFractureModeSettings>();
+	ModeSettings->ApplyDefaultConvexSettings(*FracturedGeometryCollection->GetGeometryCollection());
+
 	return NewActor;
 }
 
@@ -448,6 +453,9 @@ void UFractureToolResetAsset::Execute(TWeakPtr<FFractureEditorModeToolkit> InToo
 				FGeometryCollectionProximityUtility ProximityUtility(GeometryCollection);
 				ProximityUtility.UpdateProximity();
 
+				const UFractureModeSettings* ModeSettings = GetDefault<UFractureModeSettings>();
+				ModeSettings->ApplyDefaultConvexSettings(*GeometryCollection);
+
 				FGeometryCollectionClusteringUtility::UpdateHierarchyLevelOfChildren(GeometryCollection, -1);
 				AddSingleRootNodeIfRequired(GeometryCollectionObject);
 				GeometryCollectionComponent->MarkRenderDynamicDataDirty();
diff --git a/Engine/Plugins/Experimental/ChaosEditor/Source/FractureEditor/Public/FractureModeSettings.h b/Engine/Plugins/Experimental/ChaosEditor/Source/FractureEditor/Public/FractureModeSettings.h
new file mode 100644
--- /dev/null
+++ b/Engine/Plugins/Experimental/ChaosEditor/Source/FractureEditor/Public/FractureModeSettings.h
@@ -0,0 +1,63 @@
+// Copyright Epic Games, Inc. All Rights Reserved.
+
+#pragma once
+
+#include "Engine/DeveloperSettings.h"
+#include "Engine/EngineTypes.h"
+#include "FractureModeSettings.generated.h"
+
+class FGeometryCollection;
+
+/**
+ * Settings for the Fracture Editor Mode.
+ */
+UCLASS(config = Editor)
+class FRACTUREEDITOR_API UFractureModeSettings : public UDeveloperSettings
+{
+	GENERATED_BODY()
+
+public:
+
+	// UDeveloperSettings overrides
+
+	virtual FName GetContainerName() const override { return FName("Project"); }
+	virtual FName GetCategoryName() const override { return FName("Plugins"); }
+	virtual FName GetSectionName() const override { return FName("FractureMode"); }
+
+	virtual FText GetSectionText() const override { return NSLOCTEXT("FractureModeSettings", "FractureModeSettingsName", "Fracture Mode"); }
+	virtual FText GetSectionDescription() const override { return NSLOCTEXT("FractureModeSettings", "FractureModeSettingsDescription", "Configure the Fracture Editor Mode plugin"); }
+
+public:
+
+	/** Default fraction of convex hulls for a transform that we can remove before instead using the hulls of the children */
+	UPROPERTY(config, EditAnywhere, Category = "Fracture Mode|Convex Generation Defaults", meta = (DisplayName = "Can Remove Fraction", ClampMin = ".01", ClampMax = "1"))
+	double ConvexFractionAllowRemove = .5;
+
+	/** Default fraction of geometry volume by which a cluster's convex hull volume can exceed the actual geometry volume before instead using the hulls of the children.  0 means the convex volume cannot exceed the geometry volume; 1 means the convex volume is allowed to be 100% larger (2x) the geometry volume. */
+	UPROPERTY(config, EditAnywhere, Category = "Fracture Mode|Convex Generation Defaults", meta = (ClampMin = "0"))
+	float ConvexCanExceedFraction = .5;
+
+	/** Default simplification threshold for convex hulls of new geometry collections */
+	UPROPERTY(config, EditAnywhere, Category = "Fracture Mode|Convex Generation Defaults", meta = (ClampMin = "0"))
+	float ConvexSimplificationDistanceThreshold = 10;
+
+	/** Default overlap removal setting for convex hulls of new geometry collections */
+	UPROPERTY(config, EditAnywhere, Category = "Fracture Mode|Convex Generation Defaults")
+	bool bConvexRemoveOverlaps = true;
+
+	// Apply Convex Generation Defaults to a GeometryCollection
+	void ApplyDefaultConvexSettings(FGeometryCollection& GeometryCollection) const;
+
+
+	DECLARE_MULTICAST_DELEGATE_TwoParams(UFractureModeSettingsModified, UObject*, FProperty*);
+	UFractureModeSettingsModified OnModified;
+
+	virtual void PostEditChangeProperty(FPropertyChangedEvent& PropertyChangedEvent) override
+	{
+		OnModified.Broadcast(this, PropertyChangedEvent.Property);
+
+		Super::PostEditChangeProperty(PropertyChangedEvent);
+	}
+
+
+};
\ No newline at end of file
