fix Horde Issue #180657 -- incorrect loctext key for UFractureToolSelectCluster::GetTooltipText

diff --git a/Engine/Plugins/Experimental/ChaosEditor/Source/FractureEditor/Private/FractureToolSelectors.cpp b/Engine/Plugins/Experimental/ChaosEditor/Source/FractureEditor/Private/FractureToolSelectors.cpp
--- a/Engine/Plugins/Experimental/ChaosEditor/Source/FractureEditor/Private/FractureToolSelectors.cpp
+++ b/Engine/Plugins/Experimental/ChaosEditor/Source/FractureEditor/Private/FractureToolSelectors.cpp
@@ -345,7 +345,7 @@ FText UFractureToolSelectCluster::GetDisplayText() const
 
 FText UFractureToolSelectCluster::GetTooltipText() const
 {
-	return FText(NSLOCTEXT("Fracture", "FractureToolSelectLeafTooltip", "Select cluster bones"));
+	return FText(NSLOCTEXT("Fracture", "FractureToolSelectClusterTooltip", "Select cluster bones"));
 }
 
 FSlateIcon UFractureToolSelectCluster::GetToolIcon() const
