fix fracture tool ConvexFractionAllowRemove / FractionAllowRemove comment

diff --git a/Engine/Plugins/Experimental/ChaosEditor/Source/FractureEditor/Private/FractureToolConvex.h b/Engine/Plugins/Experimental/ChaosEditor/Source/FractureEditor/Private/FractureToolConvex.h
--- a/Engine/Plugins/Experimental/ChaosEditor/Source/FractureEditor/Private/FractureToolConvex.h
+++ b/Engine/Plugins/Experimental/ChaosEditor/Source/FractureEditor/Private/FractureToolConvex.h
@@ -20,7 +20,7 @@ public:
 		: Super(ObjInit)
 	{}
 
-	/** Fraction of the convex hulls for a transform that we can remove before instead using the hulls of the children */
+	/** Fraction of the convex hulls for a transform that we can remove before using the hulls of the children */
 	UPROPERTY(EditAnywhere, Category = Automatic, meta = (DisplayName = "Can Remove Fraction", ClampMin = ".01", ClampMax = "1"))
 	double FractionAllowRemove = .5;
 
diff --git a/Engine/Plugins/Experimental/ChaosEditor/Source/FractureEditor/Public/FractureModeSettings.h b/Engine/Plugins/Experimental/ChaosEditor/Source/FractureEditor/Public/FractureModeSettings.h
--- a/Engine/Plugins/Experimental/ChaosEditor/Source/FractureEditor/Public/FractureModeSettings.h
+++ b/Engine/Plugins/Experimental/ChaosEditor/Source/FractureEditor/Public/FractureModeSettings.h
@@ -29,7 +29,7 @@ public:
 
 public:
 
-	/** Default fraction of convex hulls for a transform that we can remove before instead using the hulls of the children */
+	/** Default fraction of convex hulls for a transform that we can remove before using the hulls of the children */
 	UPROPERTY(config, EditAnywhere, Category = "Fracture Mode|Convex Generation Defaults", meta = (DisplayName = "Can Remove Fraction", ClampMin = ".01", ClampMax = "1"))
 	double ConvexFractionAllowRemove = .5;
 
