make fracture mode only auto-scroll to bones in the outliner/histogram if they are single bone selection in the 3D view (fixes outliner view jumping away when multi-selecting in the outliner itself)

diff --git a/Engine/Plugins/Experimental/ChaosEditor/Source/FractureEditor/Private/FractureEditorMode.cpp b/Engine/Plugins/Experimental/ChaosEditor/Source/FractureEditor/Private/FractureEditorMode.cpp
--- a/Engine/Plugins/Experimental/ChaosEditor/Source/FractureEditor/Private/FractureEditorMode.cpp
+++ b/Engine/Plugins/Experimental/ChaosEditor/Source/FractureEditor/Private/FractureEditorMode.cpp
@@ -209,7 +209,7 @@ bool UFractureEditorMode::SelectFromClick(HHitProxy* HitProxy, bool bCtrlDown, b
 			if (Toolkit.IsValid())
 			{
 				FFractureEditorModeToolkit* FractureToolkit = (FFractureEditorModeToolkit*)Toolkit.Get();
-				FractureToolkit->SetBoneSelection(GeometryCollectionProxy->Component, GeometryCollectionProxy->Component->GetSelectedBones(), true);
+				FractureToolkit->SetBoneSelection(GeometryCollectionProxy->Component, GeometryCollectionProxy->Component->GetSelectedBones(), true, GeometryCollectionProxy->BoneIndex);
 			}
 
 			return true;
diff --git a/Engine/Plugins/Experimental/ChaosEditor/Source/FractureEditor/Private/FractureEditorModeToolkit.cpp b/Engine/Plugins/Experimental/ChaosEditor/Source/FractureEditor/Private/FractureEditorModeToolkit.cpp
--- a/Engine/Plugins/Experimental/ChaosEditor/Source/FractureEditor/Private/FractureEditorModeToolkit.cpp
+++ b/Engine/Plugins/Experimental/ChaosEditor/Source/FractureEditor/Private/FractureEditorModeToolkit.cpp
@@ -1478,15 +1478,15 @@ void FFractureEditorModeToolkit::SetOutlinerComponents(const TArray<UGeometryCol
 	}
 }
 
-void FFractureEditorModeToolkit::SetBoneSelection(UGeometryCollectionComponent* InRootComponent, const TArray<int32>& InSelectedBones, bool bClearCurrentSelection)
+void FFractureEditorModeToolkit::SetBoneSelection(UGeometryCollectionComponent* InRootComponent, const TArray<int32>& InSelectedBones, bool bClearCurrentSelection, int32 FocusBoneIdx)
 {
 	if (InSelectedBones.Num() > 0 && !InRootComponent->IsSelected())
 	{
 		GEditor->SelectComponent(InRootComponent, true, true);
 	}
 
-	OutlinerView->SetBoneSelection(InRootComponent, InSelectedBones, bClearCurrentSelection);
-	HistogramView->SetBoneSelection(InRootComponent, InSelectedBones, bClearCurrentSelection);
+	OutlinerView->SetBoneSelection(InRootComponent, InSelectedBones, bClearCurrentSelection, FocusBoneIdx);
+	HistogramView->SetBoneSelection(InRootComponent, InSelectedBones, bClearCurrentSelection, FocusBoneIdx);
 
 	UpdateHideForComponent(InRootComponent);
 	
diff --git a/Engine/Plugins/Experimental/ChaosEditor/Source/FractureEditor/Private/SGeometryCollectionHistogram.cpp b/Engine/Plugins/Experimental/ChaosEditor/Source/FractureEditor/Private/SGeometryCollectionHistogram.cpp
--- a/Engine/Plugins/Experimental/ChaosEditor/Source/FractureEditor/Private/SGeometryCollectionHistogram.cpp
+++ b/Engine/Plugins/Experimental/ChaosEditor/Source/FractureEditor/Private/SGeometryCollectionHistogram.cpp
@@ -361,7 +361,7 @@ void SGeometryCollectionHistogram::SetListIndices()
 }
 
 
-void SGeometryCollectionHistogram::SetBoneSelection(UGeometryCollectionComponent* RootComponent, const TArray<int32>& InSelection, bool bClearCurrentSelection)
+void SGeometryCollectionHistogram::SetBoneSelection(UGeometryCollectionComponent* RootComponent, const TArray<int32>& InSelection, bool bClearCurrentSelection, int32 FocusBoneIdx)
 {
 	TGuardValue<bool> ExternalSelectionGuard(bPerformingSelection, true);
 
@@ -381,7 +381,7 @@ void SGeometryCollectionHistogram::SetBoneSelection(UGeometryCollectionComponent
 				FGeometryCollectionHistogramItemPtr Item = RootNode->GetItemFromBoneIndex(BoneIndex);
 				if (Item.IsValid())
 				{
-					if (bFirstSelection)
+					if (bFirstSelection && BoneIndex == FocusBoneIdx)
 					{
 						ListView->RequestScrollIntoView(Item);
 						bFirstSelection = false;
diff --git a/Engine/Plugins/Experimental/ChaosEditor/Source/FractureEditor/Private/SGeometryCollectionHistogram.h b/Engine/Plugins/Experimental/ChaosEditor/Source/FractureEditor/Private/SGeometryCollectionHistogram.h
--- a/Engine/Plugins/Experimental/ChaosEditor/Source/FractureEditor/Private/SGeometryCollectionHistogram.h
+++ b/Engine/Plugins/Experimental/ChaosEditor/Source/FractureEditor/Private/SGeometryCollectionHistogram.h
@@ -141,7 +141,7 @@ public:
 
 	TSharedRef<ITableRow> MakeHistogramRowWidget(FGeometryCollectionHistogramItemPtr InInfo, const TSharedRef<STableViewBase>& OwnerTable);
 	void SetComponents(const TArray<UGeometryCollectionComponent*>& InNewComponents, int32 LevelView);
-	void SetBoneSelection(UGeometryCollectionComponent* RootComponent, const TArray<int32>& InSelection, bool bClearCurrentSelection);
+	void SetBoneSelection(UGeometryCollectionComponent* RootComponent, const TArray<int32>& InSelection, bool bClearCurrentSelection, int32 FocusBoneIdx = -1);
 	void InspectAttribute(EInspectedAttributeEnum InspectedAttribute);
 	void RefreshView(bool bSorted);
 	void RegenerateNodes(int32 LevelView);
diff --git a/Engine/Plugins/Experimental/ChaosEditor/Source/FractureEditor/Private/SGeometryCollectionOutliner.cpp b/Engine/Plugins/Experimental/ChaosEditor/Source/FractureEditor/Private/SGeometryCollectionOutliner.cpp
--- a/Engine/Plugins/Experimental/ChaosEditor/Source/FractureEditor/Private/SGeometryCollectionOutliner.cpp
+++ b/Engine/Plugins/Experimental/ChaosEditor/Source/FractureEditor/Private/SGeometryCollectionOutliner.cpp
@@ -431,7 +431,7 @@ int32 SGeometryCollectionOutliner::GetBoneSelectionCount() const
 	return TreeView->GetSelectedItems().Num();
 }
 
-void SGeometryCollectionOutliner::SetBoneSelection(UGeometryCollectionComponent* RootComponent, const TArray<int32>& InSelection, bool bClearCurrentSelection)
+void SGeometryCollectionOutliner::SetBoneSelection(UGeometryCollectionComponent* RootComponent, const TArray<int32>& InSelection, bool bClearCurrentSelection, int32 FocusBoneIdx)
 {
 	TGuardValue<bool> ExternalSelectionGuard(bPerformingSelection, true);
 
@@ -452,7 +452,7 @@ void SGeometryCollectionOutliner::SetBoneSelection(UGeometryCollectionComponent*
 				FGeometryCollectionTreeItemPtr Item = RootNode->GetItemFromBoneIndex(BoneIndex);
 				if (ensure(Item.IsValid()))
 				{
-					if (bFirstSelection)
+					if (bFirstSelection && FocusBoneIdx == BoneIndex)
 					{
 						TreeView->RequestScrollIntoView(Item);
 						bFirstSelection = false;
diff --git a/Engine/Plugins/Experimental/ChaosEditor/Source/FractureEditor/Private/SGeometryCollectionOutliner.h b/Engine/Plugins/Experimental/ChaosEditor/Source/FractureEditor/Private/SGeometryCollectionOutliner.h
--- a/Engine/Plugins/Experimental/ChaosEditor/Source/FractureEditor/Private/SGeometryCollectionOutliner.h
+++ b/Engine/Plugins/Experimental/ChaosEditor/Source/FractureEditor/Private/SGeometryCollectionOutliner.h
@@ -268,7 +268,7 @@ public:
 
 	void UpdateGeometryCollection();
 	void SetComponents(const TArray<UGeometryCollectionComponent*>& InNewComponents);
-	void SetBoneSelection(UGeometryCollectionComponent* RootComponent, const TArray<int32>& InSelection, bool bClearCurrentSelection);
+	void SetBoneSelection(UGeometryCollectionComponent* RootComponent, const TArray<int32>& InSelection, bool bClearCurrentSelection, int32 FocusBoneIdx = -1);
 	int32 GetBoneSelectionCount() const;
 
 	void SetInitialDynamicState(int32 InDynamicState);
diff --git a/Engine/Plugins/Experimental/ChaosEditor/Source/FractureEditor/Public/FractureEditorModeToolkit.h b/Engine/Plugins/Experimental/ChaosEditor/Source/FractureEditor/Public/FractureEditorModeToolkit.h
--- a/Engine/Plugins/Experimental/ChaosEditor/Source/FractureEditor/Public/FractureEditorModeToolkit.h
+++ b/Engine/Plugins/Experimental/ChaosEditor/Source/FractureEditor/Public/FractureEditorModeToolkit.h
@@ -121,7 +121,7 @@ public:
 	void Shutdown();
 
 	void SetOutlinerComponents(const TArray<UGeometryCollectionComponent*>& InNewComponents);
-	void SetBoneSelection(UGeometryCollectionComponent* InRootComponent, const TArray<int32>& InSelectedBones, bool bClearCurrentSelection);
+	void SetBoneSelection(UGeometryCollectionComponent* InRootComponent, const TArray<int32>& InSelectedBones, bool bClearCurrentSelection, int32 FocusBoneIdx = -1);
 
 	// View Settings
 	float GetExplodedViewValue() const;
