Fix crash when deleting the asset of an selected instanciated geometry collection

diff --git a/Engine/Plugins/Experimental/ChaosEditor/Source/FractureEditor/Private/FractureEditorMode.cpp b/Engine/Plugins/Experimental/ChaosEditor/Source/FractureEditor/Private/FractureEditorMode.cpp
--- a/Engine/Plugins/Experimental/ChaosEditor/Source/FractureEditor/Private/FractureEditorMode.cpp
+++ b/Engine/Plugins/Experimental/ChaosEditor/Source/FractureEditor/Private/FractureEditorMode.cpp
@@ -472,12 +472,16 @@ void UFractureEditorMode::OnActorSelectionChanged(const TArray<UObject*>& NewSel
 		{
 			GeometryCollectionComponent->SetEmbeddedGeometrySelectable(true);
 			
-			// If collection does not already have guids, make them
 			FGeometryCollectionEdit RestCollectionEdit = GeometryCollectionComponent->EditRestCollection(GeometryCollection::EEditUpdate::None);
-			if (!ensureMsgf(RestCollectionEdit.GetRestCollection(), TEXT("UGeometryCollectionComponent had no UGeometryCollection")) ||
-				!ensureMsgf(RestCollectionEdit.GetRestCollection()->GetGeometryCollection(), TEXT("UGeometryCollectionComponent had no FGeometryCollection")))
+			if (!RestCollectionEdit.GetRestCollection())
 			{
-				// guard against the component not having a UGeometryCollection or FGeometryCollection, with ensures because it doesn't seem like this should happen
+				// guard against the component not having a UGeometryCollection, thsi can happen if the asset has been force deleted while the instance was selected
+				continue;
+			}
+
+			if (!ensureMsgf(RestCollectionEdit.GetRestCollection()->GetGeometryCollection(), TEXT("UGeometryCollectionComponent had no FGeometryCollection")))
+			{
+				// guard against the component not having a FGeometryCollection, with ensures because it doesn't seem like this should happen
 				continue;
 			}
 			::GeometryCollection::GenerateTemporaryGuids(RestCollectionEdit.GetRestCollection()->GetGeometryCollection().Get());
