AutoCluster improvements:
 - Add option to choose number of Voronoi sites as a fraction of total children to cluster, instead of directly specifying the number
 - Add options to (try to) prevent creating isolated children, by merging such clusters with neighbors, and clusters of all children (by skipping cluster creation in that case)

diff --git a/Engine/Plugins/Experimental/ChaosEditor/Source/FractureEditor/Private/FractureToolAutoCluster.cpp b/Engine/Plugins/Experimental/ChaosEditor/Source/FractureEditor/Private/FractureToolAutoCluster.cpp
--- a/Engine/Plugins/Experimental/ChaosEditor/Source/FractureEditor/Private/FractureToolAutoCluster.cpp
+++ b/Engine/Plugins/Experimental/ChaosEditor/Source/FractureEditor/Private/FractureToolAutoCluster.cpp
@@ -79,32 +79,50 @@ void UFractureToolAutoCluster::Execute(TWeakPtr<FFractureEditorModeToolkit> InTo
 			for (const int32 ClusterIndex : Context.GetSelection())
 			{
 				FVoronoiPartitioner VoronoiPartition(GeometryCollection, ClusterIndex);
-				VoronoiPartition.KMeansPartition(AutoClusterSettings->SiteCount);
+				int32 NumChildren = GeometryCollection->Children[ClusterIndex].Num();
+				int32 SiteCountToUse = AutoClusterSettings->ClusterSizeMethod == EClusterSizeMethod::ByFractionOfInput ?
+					FMath::Max(2, NumChildren * AutoClusterSettings->SiteCountFraction) :
+					AutoClusterSettings->SiteCount;
+				VoronoiPartition.KMeansPartition(SiteCountToUse);
 				if (VoronoiPartition.GetPartitionCount() == 0)
 				{
 					continue;
 				}
 
-				if (AutoClusterSettings->bEnforceConnectivity)
+				bool bNeedProximity = AutoClusterSettings->bEnforceConnectivity || AutoClusterSettings->bMergeIsolatedChildren;
+				if (bNeedProximity)
 				{
 					FGeometryCollectionProximityUtility ProximityUtility(GeometryCollection);
 					ProximityUtility.UpdateProximity();
+				}
+
+				if (AutoClusterSettings->bEnforceConnectivity)
+				{
 					VoronoiPartition.SplitDisconnectedPartitions(GeometryCollection);
 				}
+
+				if (AutoClusterSettings->bMergeIsolatedChildren)
+				{
+					VoronoiPartition.MergeSingleElementPartitions(GeometryCollection);
+				}
+
+				int32 NonEmptyPartitionCount = VoronoiPartition.GetNonEmptyPartitionCount();
+				bool bHasEmptyClusters = NonEmptyPartitionCount > 0;
+
+				if (AutoClusterSettings->bDoNotCreateSingleSiteClusters && NonEmptyPartitionCount == 1)
+				{
+					continue;
+				}
 					
 				int32 PartitionCount = VoronoiPartition.GetPartitionCount();
 				int32 NewClusterIndexStart = GeometryCollection->AddElements(PartitionCount, FGeometryCollection::TransformGroup);
 
-				bool bHasEmptyClusters = false;
+
 
 				for (int32 Index = 0; Index < PartitionCount; ++Index)
 				{
 						
 					TArray<int32> NewCluster = VoronoiPartition.GetPartition(Index);
-					if (NewCluster.Num() == 0)
-					{
-						bHasEmptyClusters = true;
-					}
 
 					int32 NewClusterIndex = NewClusterIndexStart + Index;
 					GeometryCollection->Parent[NewClusterIndex] = ClusterIndex;
@@ -157,10 +175,52 @@ void FVoronoiPartitioner::KMeansPartition(int32 InPartitionCount)
 	}
 }
 
+void FVoronoiPartitioner::MergeSingleElementPartitions(FGeometryCollection* GeometryCollection)
+{
+	if (Connectivity.IsEmpty())
+	{
+		GenerateConnectivity(GeometryCollection);
+	}
+	for (int32 ElIdx = 0; ElIdx < TransformIndices.Num(); ElIdx++)
+	{
+		int32 Partition = Partitions[ElIdx];
+		if (PartitionSize[Partition] == 1)
+		{
+			// Find the smallest neighboring partition to merge to
+			// (to help keep partition sizes balanced)
+			int32 SmallestNbrPartition = -1;
+			int32 SmallestNbrSize = TransformIndices.Num()+1;
+			for (int32 NbrEl : Connectivity[ElIdx])
+			{
+				int32 NbrPartition = Partitions[NbrEl];
+				if (NbrPartition == Partition)
+				{
+					continue;
+				}
+				int32 NbrSize = PartitionSize[NbrPartition];
+				if (SmallestNbrPartition == -1 || NbrSize < SmallestNbrSize)
+				{
+					SmallestNbrPartition = NbrPartition;
+					SmallestNbrSize = NbrSize;
+				}
+			}
+			if (SmallestNbrPartition != -1)
+			{
+				Partitions[ElIdx] = SmallestNbrPartition;
+				PartitionSize[Partition]--;
+				PartitionSize[SmallestNbrPartition]++;
+			}
+		}
+	}
+}
+
 void FVoronoiPartitioner::SplitDisconnectedPartitions(FGeometryCollection* GeometryCollection)
 {
 	Visited.Init(false, TransformIndices.Num());
-	GenerateConnectivity(GeometryCollection);
+	if (Connectivity.IsEmpty())
+	{
+		GenerateConnectivity(GeometryCollection);
+	}
 	
 	for (int32 PartitionIndex = 0; PartitionIndex < PartitionCount; ++PartitionIndex)
 	{
diff --git a/Engine/Plugins/Experimental/ChaosEditor/Source/FractureEditor/Private/FractureToolAutoCluster.h b/Engine/Plugins/Experimental/ChaosEditor/Source/FractureEditor/Private/FractureToolAutoCluster.h
--- a/Engine/Plugins/Experimental/ChaosEditor/Source/FractureEditor/Private/FractureToolAutoCluster.h
+++ b/Engine/Plugins/Experimental/ChaosEditor/Source/FractureEditor/Private/FractureToolAutoCluster.h
@@ -3,6 +3,7 @@
 #pragma once
 
 #include "FractureTool.h"
+#include "Algo/Count.h"
 
 #include "FractureToolAutoCluster.generated.h"
 
@@ -22,6 +23,15 @@ enum class EFractureAutoClusterMode : uint8
 	Voronoi UMETA(DisplayName = "Voronoi"),
 };
 
+UENUM()
+enum class EClusterSizeMethod : uint8
+{
+	// Cluster by specifying an absolute number of clusters
+	ByNumber,
+	// Cluster by specifying a fraction of the number of input bones
+	ByFractionOfInput
+};
+
 
 UCLASS(DisplayName = "Auto Cluster", Category = "FractureTools")
 class UFractureAutoClusterSettings : public UFractureToolSettings
@@ -37,13 +47,29 @@ public:
 	UPROPERTY(meta = (DeprecatedProperty, DeprecationMessage = "Simplified interface now only supports Voronoi clustering."))
 	EFractureAutoClusterMode AutoClusterMode_DEPRECATED;
 
+	/** How to choose the size of the clusters to create */
+	UPROPERTY(EditAnywhere, Category = ClusterSize)
+	EClusterSizeMethod ClusterSizeMethod = EClusterSizeMethod::ByNumber;
+
 	/** Use a Voronoi diagram with this many Voronoi sites as a guide for deciding cluster boundaries */
-	UPROPERTY(EditAnywhere, Category = AutoCluster, meta = (DisplayName = "Cluster Sites", UIMin = "1", UIMax = "5000", ClampMin = "1"))
+	UPROPERTY(EditAnywhere, Category = ClusterSize, meta = (DisplayName = "Cluster Sites", UIMin = "2", UIMax = "5000", ClampMin = "1", EditCondition = "ClusterSizeMethod == EClusterSizeMethod::ByNumber"))
 	uint32 SiteCount=10;
 
+	/** Choose the number of Voronoi sites used for clustering as a fraction of the number of child bones to process */
+	UPROPERTY(EditAnywhere, Category = ClusterSize, meta = (DisplayName = "Cluster Fraction", ClampMin = "0", ClampMax = ".5", EditCondition = "ClusterSizeMethod == EClusterSizeMethod::ByFractionOfInput"))
+	float SiteCountFraction = .5;
+
 	/** If true, bones will only be added to the same cluster if they are physically connected (either directly, or via other bones in the same cluster) */
 	UPROPERTY(EditAnywhere, Category = AutoCluster, meta = (DisplayName = "Enforce Cluster Connectivity"))
 	bool bEnforceConnectivity=true;
+
+	/** If true, tries to prevent the creation of clusters with only a single child, by merging such clusters into a neighbor cluster if possible */
+	UPROPERTY(EditAnywhere, Category = AutoCluster)
+	bool bMergeIsolatedChildren = true;
+
+	/** If true, will not create a new cluster if all children would fall in that same cluster -- for example, if there is only one Voronoi site */
+	UPROPERTY(EditAnywhere, Category = AutoCluster)
+	bool bDoNotCreateSingleSiteClusters = true;
 };
 
 
@@ -81,8 +107,16 @@ public:
 	/** Split any partition islands into their own partition. This will possbily increase number of partitions to exceed desired count. */
 	void SplitDisconnectedPartitions(FGeometryCollection* GeometryCollection);
 
+	/** Merge any partitions w/ only 1 body into a connected, neighboring partition (if any).  This can decrease the number of partitions below the desired count. */
+	void MergeSingleElementPartitions(FGeometryCollection* GeometryCollection);
+
 	int32 GetPartitionCount() const { return PartitionCount; }
 
+	int32 GetNonEmptyPartitionCount() const
+	{
+		return PartitionSize.Num() - Algo::Count(PartitionSize, 0);
+	}
+
 	/** return the GeometryCollection TranformIndices within the partition. */
 	TArray<int32> GetPartition(int32 PartitionIndex) const;
 
