Undoing the creation of an in-editor only (not saved yet) Geometry Collection will now correctly update the Content Browser.

diff --git a/Engine/Plugins/Experimental/ChaosEditor/Source/FractureEditor/Private/FractureToolGenerators.cpp b/Engine/Plugins/Experimental/ChaosEditor/Source/FractureEditor/Private/FractureToolGenerators.cpp
--- a/Engine/Plugins/Experimental/ChaosEditor/Source/FractureEditor/Private/FractureToolGenerators.cpp
+++ b/Engine/Plugins/Experimental/ChaosEditor/Source/FractureEditor/Private/FractureToolGenerators.cpp
@@ -12,6 +12,8 @@
 #include "Engine/SkeletalMesh.h"
 #include "Engine/StaticMesh.h"
 #include "Engine/StaticMeshActor.h"
+#include "Misc/Change.h"
+#include "ChangeTransactor.h"
 
 #include "FractureModeSettings.h"
 
@@ -25,9 +27,28 @@
 #include "GeometryCollection/GeometryCollectionProximityUtility.h"
 #include "FractureToolContext.h"
 
-
 #define LOCTEXT_NAMESPACE "FractureToolGenerators"
 
+// Creates an undo/redo action that (un)registers and object with the Asset Registry.
+// Upon undo this causes the object to be unregistered and as a result be removed from
+// Content Browsers.
+class FAssetRegistrationChange final : public FCommandChange
+{
+public:
+	void Apply(UObject* Object) override
+	{
+		FAssetRegistryModule::AssetCreated(Object);
+	}
+	void Revert(UObject* Object) override
+	{
+		FAssetRegistryModule::AssetDeleted(Object);
+	}
+
+	FString ToString() const override
+	{
+		return TEXT("Asset registry from " LOCTEXT_NAMESPACE);
+	}
+};
 
 FText UFractureToolGenerateAsset::GetDisplayText() const
 {
@@ -274,7 +295,6 @@ AGeometryCollectionActor* UFractureToolGenerateAsset::ConvertActorsToGeometryCol
 
 class AGeometryCollectionActor* UFractureToolGenerateAsset::CreateNewGeometryActor(const FString& InAssetPath, const FTransform& Transform, bool AddMaterials /*= false*/)
 {
-
 	FString UniquePackageName = InAssetPath;
 	FString UniqueAssetName = FPackageName::GetLongPackageAssetName(InAssetPath);
 
@@ -285,6 +305,12 @@ class AGeometryCollectionActor* UFractureToolGenerateAsset::CreateNewGeometryAct
 	UGeometryCollection* InGeometryCollection = static_cast<UGeometryCollection*>(NewObject<UGeometryCollection>(Package, UGeometryCollection::StaticClass(), FName(*UniqueAssetName), RF_Transactional | RF_Public | RF_Standalone));
 	if(!InGeometryCollection->SizeSpecificData.Num()) InGeometryCollection->SizeSpecificData.Add(FGeometryCollectionSizeSpecificData());
 
+	// Record the creation of the geometry collection so it's removed from the Asset Registry and the Content Browser when undo is called.
+	UE::FChangeTransactor Transactor(InGeometryCollection);
+	Transactor.OpenTransaction(LOCTEXT("GeometryCollectionAssetRegistration", "Geometry Collection Asset Registration"));
+	Transactor.AddTransactionChange<FAssetRegistrationChange>();
+	Transactor.CloseTransaction();
+
 	// Create the new Geometry Collection actor
 	AGeometryCollectionActor* NewActor = Cast<AGeometryCollectionActor>(AddActor(GetSelectedLevel(), AGeometryCollectionActor::StaticClass()));
 	check(NewActor->GetGeometryCollectionComponent());
