fix fracture autouv layout deleting the geometry collection on failure to do layout / do not generate a default 'empty' result for fracture background op tasks

diff --git a/Engine/Plugins/Experimental/ChaosEditor/Source/FractureEditor/Private/FractureToolBackgroundTask.h b/Engine/Plugins/Experimental/ChaosEditor/Source/FractureEditor/Private/FractureToolBackgroundTask.h
--- a/Engine/Plugins/Experimental/ChaosEditor/Source/FractureEditor/Private/FractureToolBackgroundTask.h
+++ b/Engine/Plugins/Experimental/ChaosEditor/Source/FractureEditor/Private/FractureToolBackgroundTask.h
@@ -21,7 +21,7 @@ public:
 	int ResultGeometryIndex = -1;
 	TUniquePtr<FGeometryCollection> CollectionCopy;
 
-	TGeometryCollectionOperator(const FGeometryCollection& SourceCollection)
+	TGeometryCollectionOperator(const FGeometryCollection& SourceCollection) : UE::Geometry::TGenericDataOperator<ResultType>(false)
 	{
 		CollectionCopy = MakeUnique<FGeometryCollection>();
 		CollectionCopy->CopyMatchingAttributesFrom(SourceCollection, nullptr);
