add standard fracture settings (including random seed to control the cut distribution) back to the Mesh fracture tool, except not Grout and Noise (as these aren't supported)

diff --git a/Engine/Plugins/Experimental/ChaosEditor/Source/FractureEditor/Private/FractureToolCutter.h b/Engine/Plugins/Experimental/ChaosEditor/Source/FractureEditor/Private/FractureToolCutter.h
--- a/Engine/Plugins/Experimental/ChaosEditor/Source/FractureEditor/Private/FractureToolCutter.h
+++ b/Engine/Plugins/Experimental/ChaosEditor/Source/FractureEditor/Private/FractureToolCutter.h
@@ -68,36 +68,44 @@ public:
 	bool bDrawDiagram;
 
 	/** Amount of space to leave between cut pieces */
-	UPROPERTY(EditAnywhere, Category = CommonFracture, meta = (UIMin = "0.0", ClampMin = "0.0"))
+	UPROPERTY(EditAnywhere, Category = CommonFracture, meta = (UIMin = "0.0", ClampMin = "0.0", EditCondition = "bGroutSettingEnabled", HideEditConditionToggle, EditConditionHides))
 	float Grout = 0.0f;
 
+	// This flag allows tools to disable the above Noise settings if/when they are not applicable
+	UPROPERTY()
+	bool bGroutSettingEnabled = true;
+
 	/** Size of the Perlin noise displacement (in cm). If 0, no noise will be applied */
-	UPROPERTY(EditAnywhere, Category = Noise, meta = (UIMin = "0.0"))
+	UPROPERTY(EditAnywhere, Category = Noise, meta = (UIMin = "0.0", EditCondition = "bNoiseSettingsEnabled", HideEditConditionToggle, EditConditionHides))
 	float Amplitude;
 
 	/** Period of the Perlin noise.  Smaller values will create a smoother noise pattern */
-	UPROPERTY(EditAnywhere, Category = Noise)
+	UPROPERTY(EditAnywhere, Category = Noise, meta = (EditCondition = "bNoiseSettingsEnabled", HideEditConditionToggle, EditConditionHides))
 	float Frequency;
 
 	/** Persistence of the layers of Perlin noise. At each layer (octave) after the first, the amplitude of the Perlin noise is scaled by this factor */
-	UPROPERTY(EditAnywhere, Category = Noise, meta = (UIMin = "0.01", UIMax = "1.0"))
+	UPROPERTY(EditAnywhere, Category = Noise, meta = (UIMin = "0.01", UIMax = "1.0", EditCondition = "bNoiseSettingsEnabled", HideEditConditionToggle, EditConditionHides))
 	float Persistence;
 
 	/** Lacunarity of the layers of Perlin noise. At each layer (octave) after the first, the frequency of the Perlin noise is scaled by this factor */
-	UPROPERTY(EditAnywhere, Category = Noise, meta = (UIMin = "1.0", UIMax = "4.0"))
+	UPROPERTY(EditAnywhere, Category = Noise, meta = (UIMin = "1.0", UIMax = "4.0", EditCondition = "bNoiseSettingsEnabled", HideEditConditionToggle, EditConditionHides))
 	float Lacunarity;
 
 	/** 
 	 * Number of fractal layers of Perlin noise to apply. Each layer is additive, with Amplitude and Frequency parameters scaled by Persistence and Lacunarity
 	 * Smaller values (1 or 2) will create noise that looks like gentle rolling hills, while larger values (> 4) will tend to look more like craggy mountains
 	 */
-	UPROPERTY(EditAnywhere, Category = Noise, meta = (ClampMin = "1", UIMax = "8"))
+	UPROPERTY(EditAnywhere, Category = Noise, meta = (ClampMin = "1", UIMax = "8", EditCondition = "bNoiseSettingsEnabled", HideEditConditionToggle, EditConditionHides))
 	int32 OctaveNumber;
 
 	/** Distance (in cm) between vertices on cut surfaces where noise is added.  Larger spacing between vertices will create more efficient meshes with fewer triangles, but less resolution to see the shape of the added noise  */
-	UPROPERTY(EditAnywhere, Category = Noise, meta = (DisplayName = "Point Spacing", UIMin = "1", ClampMin = "0.1"))
+	UPROPERTY(EditAnywhere, Category = Noise, meta = (DisplayName = "Point Spacing", UIMin = "1", ClampMin = "0.1", EditCondition = "bNoiseSettingsEnabled", HideEditConditionToggle, EditConditionHides))
 	float PointSpacing;
 
+	// This flag allows tools to disable the above Noise settings if/when they are not applicable
+	UPROPERTY()
+	bool bNoiseSettingsEnabled = true;
+
 	/// Get the maximum distance a vertex could be moved by a combination of grout and noise
 	float GetMaxVertexMovement()
 	{
@@ -183,6 +191,18 @@ public:
 		}
 	}
 
+	void DisableGroutSetting()
+	{
+		CutterSettings->bGroutSettingEnabled = false;
+		CutterSettings->Grout = 0;
+	}
+
+	void DisableNoiseSettings()
+	{
+		CutterSettings->bNoiseSettingsEnabled = false;
+		CutterSettings->Amplitude = 0;
+	}
+
 protected:
 	UPROPERTY(EditAnywhere, Category = Slicing)
 	TObjectPtr<UFractureCutterSettings> CutterSettings;
diff --git a/Engine/Plugins/Experimental/ChaosEditor/Source/FractureEditor/Private/FractureToolMeshCut.cpp b/Engine/Plugins/Experimental/ChaosEditor/Source/FractureEditor/Private/FractureToolMeshCut.cpp
--- a/Engine/Plugins/Experimental/ChaosEditor/Source/FractureEditor/Private/FractureToolMeshCut.cpp
+++ b/Engine/Plugins/Experimental/ChaosEditor/Source/FractureEditor/Private/FractureToolMeshCut.cpp
@@ -26,6 +26,9 @@ UFractureToolMeshCut::UFractureToolMeshCut(const FObjectInitializer& ObjInit)
 {
 	MeshCutSettings = NewObject<UFractureMeshCutSettings>(GetTransientPackage(), UFractureMeshCutSettings::StaticClass());
 	MeshCutSettings->OwnerTool = this;
+	CutterSettings->bDrawSitesToggleEnabled = false;
+	DisableNoiseSettings();
+	DisableGroutSetting();
 }
 
 FText UFractureToolMeshCut::GetDisplayText() const
@@ -53,7 +56,7 @@ TArray<UObject*> UFractureToolMeshCut::GetSettingsObjects() const
  {
 	TArray<UObject*> Settings;
 	Settings.Add(MeshCutSettings);
-	//Settings.Add(CutterSettings); // TODO: add cutter settings if/when we support noise and grout
+	Settings.Add(CutterSettings);
 	Settings.Add(CollisionSettings);
 	return Settings;
 }
@@ -179,6 +182,8 @@ void UFractureToolMeshCut::GenerateMeshTransforms(const FFractureToolContext& Co
 
 void UFractureToolMeshCut::FractureContextChanged()
 {
+	SetMandateGroupFracture(MeshCutSettings->CutDistribution == EMeshCutDistribution::SingleCut);
+
 	UpdateDefaultRandomSeed();
 	TArray<FFractureToolContext> FractureContexts = GetFractureToolContexts();
 
