fix autocluster crash when merging single-child clusters + properly tracking partition sizes when partitions are split to maintain proximity

diff --git a/Engine/Plugins/Experimental/ChaosEditor/Source/FractureEditor/Private/FractureToolAutoCluster.cpp b/Engine/Plugins/Experimental/ChaosEditor/Source/FractureEditor/Private/FractureToolAutoCluster.cpp
--- a/Engine/Plugins/Experimental/ChaosEditor/Source/FractureEditor/Private/FractureToolAutoCluster.cpp
+++ b/Engine/Plugins/Experimental/ChaosEditor/Source/FractureEditor/Private/FractureToolAutoCluster.cpp
@@ -198,7 +198,7 @@ void FVoronoiPartitioner::MergeSingleElementPartitions(FGeometryCollection* Geom
 					continue;
 				}
 				int32 NbrSize = PartitionSize[NbrPartition];
-				if (SmallestNbrPartition == -1 || NbrSize < SmallestNbrSize)
+				if (NbrSize > 0 && (SmallestNbrPartition == -1 || NbrSize < SmallestNbrSize))
 				{
 					SmallestNbrPartition = NbrPartition;
 					SmallestNbrSize = NbrSize;
@@ -241,8 +241,11 @@ void FVoronoiPartitioner::SplitDisconnectedPartitions(FGeometryCollection* Geome
 						{
 							bFoundUnattached = true;
 							PartitionCount++;
+							PartitionSize.Add(0);
 						}
 						Partitions[Index] = PartitionCount - 1;
+						PartitionSize[PartitionCount - 1]++;
+						PartitionSize[PartitionIndex]--;
 					}
 				}
 			}
