fix fracture tool selection not correctly scrolling to the newly selected bone with view levels other than 'all' + selection not respecting view level w/ shift+click

diff --git a/Engine/Plugins/Experimental/ChaosEditor/Source/FractureEditor/Private/FractureEditorMode.cpp b/Engine/Plugins/Experimental/ChaosEditor/Source/FractureEditor/Private/FractureEditorMode.cpp
--- a/Engine/Plugins/Experimental/ChaosEditor/Source/FractureEditor/Private/FractureEditorMode.cpp
+++ b/Engine/Plugins/Experimental/ChaosEditor/Source/FractureEditor/Private/FractureEditorMode.cpp
@@ -20,6 +20,7 @@
 #include "GeometryCollection/GeometryCollectionActor.h"
 #include "GeometryCollection/GeometryCollection.h"
 #include "GeometryCollection/GeometryCollectionUtility.h"
+#include "GeometryCollection/GeometryCollectionClusteringUtility.h"
 #include "FractureSelectionTools.h"
 #include "EditorViewportClient.h"
 #include "ScopedTransaction.h"
@@ -201,7 +202,21 @@ bool UFractureEditorMode::SelectFromClick(HHitProxy* HitProxy, bool bCtrlDown, b
 
 		if (GeometryCollectionProxy->Component)
 		{
-			TArray<int32> BoneIndices({ GeometryCollectionProxy->BoneIndex });
+			int32 BoneIndex = GeometryCollectionProxy->BoneIndex;
+			// Switch BoneIndex to match the view level
+			int32 ViewLevel = -1;
+			if (Toolkit.IsValid())
+			{
+				ViewLevel = ((FFractureEditorModeToolkit*)Toolkit.Get())->GetLevelViewValue();
+				if (ViewLevel > -1)
+				{
+					const FGeometryCollection* GeometryCollection = GeometryCollectionProxy->Component->RestCollection->GetGeometryCollection().Get();
+					BoneIndex = FGeometryCollectionClusteringUtility::GetParentOfBoneAtSpecifiedLevel(
+						GeometryCollection, BoneIndex, ViewLevel, true /*bSkipFiltered*/);
+				}
+			}
+			
+			TArray<int32> BoneIndices({ BoneIndex });
 
 			GeometryCollectionProxy->Component->Modify();
 			FFractureSelectionTools::ToggleSelectedBones(GeometryCollectionProxy->Component, BoneIndices, !(bCtrlDown || bShiftDown), bShiftDown);
@@ -209,7 +224,7 @@ bool UFractureEditorMode::SelectFromClick(HHitProxy* HitProxy, bool bCtrlDown, b
 			if (Toolkit.IsValid())
 			{
 				FFractureEditorModeToolkit* FractureToolkit = (FFractureEditorModeToolkit*)Toolkit.Get();
-				FractureToolkit->SetBoneSelection(GeometryCollectionProxy->Component, GeometryCollectionProxy->Component->GetSelectedBones(), true, GeometryCollectionProxy->BoneIndex);
+				FractureToolkit->SetBoneSelection(GeometryCollectionProxy->Component, GeometryCollectionProxy->Component->GetSelectedBones(), true, BoneIndex);
 			}
 
 			return true;
diff --git a/Engine/Source/Runtime/Experimental/GeometryCollectionEngine/Private/GeometryCollection/GeometryCollectionComponent.cpp b/Engine/Source/Runtime/Experimental/GeometryCollectionEngine/Private/GeometryCollection/GeometryCollectionComponent.cpp
--- a/Engine/Source/Runtime/Experimental/GeometryCollectionEngine/Private/GeometryCollection/GeometryCollectionComponent.cpp
+++ b/Engine/Source/Runtime/Experimental/GeometryCollectionEngine/Private/GeometryCollection/GeometryCollectionComponent.cpp
@@ -2630,7 +2630,7 @@ void FScopedColorEdit::ToggleSelectedBones(const TArray<int32>& SelectedBonesIn,
 		
 			if (bAdd) // shift select
 			{
-				Component->SelectedBones.Add(BoneIndex);
+				Component->SelectedBones.Add(ContextBoneIndex);
 			}
 			else // ctrl select (toggle)
 			{ 
