Make fracture validation tool be modal w/ options
Add option to remove 'clusters of one' w/ the validate tool

diff --git a/Engine/Plugins/Experimental/ChaosEditor/Source/FractureEditor/Private/FractureToolEditing.cpp b/Engine/Plugins/Experimental/ChaosEditor/Source/FractureEditor/Private/FractureToolEditing.cpp
--- a/Engine/Plugins/Experimental/ChaosEditor/Source/FractureEditor/Private/FractureToolEditing.cpp
+++ b/Engine/Plugins/Experimental/ChaosEditor/Source/FractureEditor/Private/FractureToolEditing.cpp
@@ -234,6 +234,12 @@ void UFractureToolUnhide::Execute(TWeakPtr<FFractureEditorModeToolkit> InToolkit
 	}
 }
 
+UFractureToolValidate::UFractureToolValidate(const FObjectInitializer& ObjInit) : Super(ObjInit)
+{
+	ValidationSettings = NewObject<UFractureValidateSettings>(GetTransientPackage(), UFractureValidateSettings::StaticClass());
+	ValidationSettings->OwnerTool = this;
+}
+
 
 FText UFractureToolValidate::GetDisplayText() const
 {
@@ -256,24 +262,33 @@ void UFractureToolValidate::RegisterUICommand(FFractureEditorCommands* BindingCo
 	BindingContext->Validate = UICommandInfo;
 }
 
+TArray<UObject*> UFractureToolValidate::GetSettingsObjects() const
+{
+	TArray<UObject*> Settings;
+	Settings.Add(ValidationSettings);
+	return Settings;
+}
+
 void UFractureToolValidate::Execute(TWeakPtr<FFractureEditorModeToolkit> InToolkit)
 {
 	if (InToolkit.IsValid())
 	{
 		FFractureEditorModeToolkit* Toolkit = InToolkit.Pin().Get();
 
+		bool bUpdatedCollections = false;
+
 		TSet<UGeometryCollectionComponent*> GeomCompSelection;
 		GetSelectedGeometryCollectionComponents(GeomCompSelection);
 		for (UGeometryCollectionComponent* GeometryCollectionComponent : GeomCompSelection)
 		{
+			bool bDirty = false;
+
 			FGeometryCollectionEdit GeometryCollectionEdit = GeometryCollectionComponent->EditRestCollection();
 			if (UGeometryCollection* GeometryCollectionObject = GeometryCollectionEdit.GetRestCollection())
 			{
 				TSharedPtr<FGeometryCollection, ESPMode::ThreadSafe> GeometryCollectionPtr = GeometryCollectionObject->GetGeometryCollection();
 				if (FGeometryCollection* GeometryCollection = GeometryCollectionPtr.Get())
 				{
-					bool bDirty = false;
-
 					TManagedArray<int32>& TransformToGeometry = GeometryCollection->TransformToGeometryIndex;
 					constexpr bool bClustersCanHaveGeometry = true;
 					if (!bClustersCanHaveGeometry)
@@ -292,31 +307,48 @@ void UFractureToolValidate::Execute(TWeakPtr<FFractureEditorModeToolkit> InToolk
 					}
 
 					// Remove any unreferenced geometry
-					TManagedArray<int32>& TransformIndex = GeometryCollection->TransformIndex;
-					const int32 GeometryCount = TransformIndex.Num();
+					if (ValidationSettings->bRemoveUnreferencedGeometry)
+					{
+						TManagedArray<int32>& TransformIndex = GeometryCollection->TransformIndex;
+						const int32 GeometryCount = TransformIndex.Num();
+
+						TArray<int32> RemoveGeometry;
+						RemoveGeometry.Reserve(GeometryCount);
+
+						for (int32 Idx = 0; Idx < GeometryCount; ++Idx)
+						{
+							if ((TransformIndex[Idx] == INDEX_NONE) || (TransformToGeometry[TransformIndex[Idx]] != Idx))
+							{
+								RemoveGeometry.Add(Idx);
+								UE_LOG(LogFractureTool, Warning, TEXT("Removed dangling geometry at index %d."), Idx);
+								bDirty = true;
+							}
+						}
 
-					TArray<int32> RemoveGeometry;
-					RemoveGeometry.Reserve(GeometryCount);
+						if (RemoveGeometry.Num() > 0)
+						{
+							FManagedArrayCollection::FProcessingParameters Params;
+							Params.bDoValidation = false; // for perf reasons
+							GeometryCollection->RemoveElements(FGeometryCollection::GeometryGroup, RemoveGeometry);
+						}
+					}
 
-					for (int32 Idx = 0; Idx < GeometryCount; ++Idx)
+					if (ValidationSettings->bRemoveClustersOfOne)
 					{
-						if ((TransformIndex[Idx] == INDEX_NONE) || (TransformToGeometry[TransformIndex[Idx]] != Idx))
+						if (FGeometryCollectionClusteringUtility::RemoveClustersOfOnlyOneChild(GeometryCollection))
 						{
-							RemoveGeometry.Add(Idx);
-							UE_LOG(LogFractureTool, Warning, TEXT("Removed dangling geometry at index %d."), Idx);
 							bDirty = true;
 						}
 					}
 
-					if (RemoveGeometry.Num() > 0)
+					if (ValidationSettings->bRemoveDanglingClusters)
 					{
-						GeometryCollection->RemoveElements(FGeometryCollection::GeometryGroup, RemoveGeometry);
+						if (FGeometryCollectionClusteringUtility::RemoveDanglingClusters(GeometryCollection))
+						{
+							bDirty = true;
+						}
 					}
 
-					// remove dangling clusters
-					// #todo leaving this out for the moment because we don't want to invalidate existing caches.
-					// FGeometryCollectionClusteringUtility::RemoveDanglingClusters(GeometryCollection);
-
 					if (bDirty)
 					{
 						FGeometryCollectionClusteringUtility::UpdateHierarchyLevelOfChildren(GeometryCollection, -1);
@@ -346,9 +378,22 @@ void UFractureToolValidate::Execute(TWeakPtr<FFractureEditorModeToolkit> InToolk
 
 			GeometryCollectionComponent->InitializeEmbeddedGeometry();
 
-			FScopedColorEdit EditBoneColor = GeometryCollectionComponent->EditBoneSelection();
-			EditBoneColor.ResetBoneSelection();
-			EditBoneColor.ResetHighlightedBones();
+			if (bDirty)
+			{
+				// reset bone selection because bones may have been deleted
+				FScopedColorEdit EditBoneColor = GeometryCollectionComponent->EditBoneSelection();
+				EditBoneColor.ResetBoneSelection();
+				EditBoneColor.ResetHighlightedBones();
+
+				// flag that at least one geometry collection has changed
+				bUpdatedCollections = true;
+			}
+		}
+
+		if (bUpdatedCollections)
+		{
+			Toolkit->RegenerateHistogram();
+			Toolkit->RegenerateOutliner();
 		}
 
 		Toolkit->OnSetLevelViewValue(-1);
diff --git a/Engine/Plugins/Experimental/ChaosEditor/Source/FractureEditor/Private/FractureToolEditing.h b/Engine/Plugins/Experimental/ChaosEditor/Source/FractureEditor/Private/FractureToolEditing.h
--- a/Engine/Plugins/Experimental/ChaosEditor/Source/FractureEditor/Private/FractureToolEditing.h
+++ b/Engine/Plugins/Experimental/ChaosEditor/Source/FractureEditor/Private/FractureToolEditing.h
@@ -57,19 +57,60 @@ public:
 };
 
 
+/** Settings for the Validate tool **/
+UCLASS(config = EditorPerProjectUserSettings)
+class UFractureValidateSettings : public UFractureToolSettings
+{
+public:
+
+	GENERATED_BODY()
+
+	UFractureValidateSettings(const FObjectInitializer& ObjInit)
+		: Super(ObjInit)
+	{
+	}
+
+	/** Find and remove any unused geometry data */
+	UPROPERTY(EditAnywhere, Category = CleanUnused)
+	bool bRemoveUnreferencedGeometry = true;
+
+	/** Whether to collapse any clusters with only a single child */
+	UPROPERTY(EditAnywhere, Category = Clustering)
+	bool bRemoveClustersOfOne = false;
+
+	/** Remove dangling clusters -- Note this can invalidate caches */
+	UPROPERTY(EditAnywhere, Category = Clustering)
+	bool bRemoveDanglingClusters = false;
+
+};
+
+
+
 
 UCLASS(DisplayName = "Validate", Category = "FractureTools")
-class UFractureToolValidate : public UFractureActionTool
+class UFractureToolValidate : public UFractureModalTool
 {
 public:
 	GENERATED_BODY()
 
-	UFractureToolValidate(const FObjectInitializer& ObjInit) : Super(ObjInit) {}
+	UFractureToolValidate(const FObjectInitializer& ObjInit);
 
 	// UFractureActionTool Interface
 	virtual FText GetDisplayText() const override;
 	virtual FText GetTooltipText() const override;
+	virtual FText GetApplyText() const override { return FText(NSLOCTEXT("ValidateTool", "ExecuteValidate", "Validate")); }
+
+	virtual TArray<UObject*> GetSettingsObjects() const override;
+
 	virtual FSlateIcon GetToolIcon() const override;
 	virtual void RegisterUICommand(FFractureEditorCommands* BindingContext) override;
 	virtual void Execute(TWeakPtr<FFractureEditorModeToolkit> InToolkit) override;
+	virtual int32 ExecuteFracture(const FFractureToolContext& FractureContext) override
+	{
+		checkNoEntry(); // Updates handled via the Execute() call instead (TODO: is that best?)
+		return -1;
+	}
+
+	UPROPERTY(EditAnywhere, Category = Validation)
+	TObjectPtr<UFractureValidateSettings> ValidationSettings;
 };
diff --git a/Engine/Plugins/Experimental/ChaosEditor/Source/FractureEditor/Private/SGeometryCollectionOutliner.cpp b/Engine/Plugins/Experimental/ChaosEditor/Source/FractureEditor/Private/SGeometryCollectionOutliner.cpp
--- a/Engine/Plugins/Experimental/ChaosEditor/Source/FractureEditor/Private/SGeometryCollectionOutliner.cpp
+++ b/Engine/Plugins/Experimental/ChaosEditor/Source/FractureEditor/Private/SGeometryCollectionOutliner.cpp
@@ -582,6 +582,10 @@ void FGeometryCollectionTreeItemComponent::GetChildrenForBone(FGeometryCollectio
 
 			if (const int32* BoneIndex = GuidIndexMap.Find(BoneItem.GetGuid()))
 			{
+				if (!ensure(*BoneIndex >= 0 && *BoneIndex < Collection->NumElements(FGeometryCollection::TransformGroup)))
+				{
+					return;
+				}
 				const TManagedArray<TSet<int32>>& Children = Collection->Children;
 				const TManagedArray<FGuid>& Guids = Collection->GetAttribute<FGuid>("GUID", "Transform");
 				for (auto ChildIndex : Children[*BoneIndex])
diff --git a/Engine/Source/Runtime/Experimental/Chaos/Private/GeometryCollection/GeometryCollectionClusteringUtility.cpp b/Engine/Source/Runtime/Experimental/Chaos/Private/GeometryCollection/GeometryCollectionClusteringUtility.cpp
--- a/Engine/Source/Runtime/Experimental/Chaos/Private/GeometryCollection/GeometryCollectionClusteringUtility.cpp
+++ b/Engine/Source/Runtime/Experimental/Chaos/Private/GeometryCollection/GeometryCollectionClusteringUtility.cpp
@@ -868,7 +868,48 @@ int32 FGeometryCollectionClusteringUtility::FindLowestCommonAncestor(FGeometryCo
 	return INDEX_NONE;
 }
 
-void FGeometryCollectionClusteringUtility::RemoveDanglingClusters(FGeometryCollection* GeometryCollection)
+bool FGeometryCollectionClusteringUtility::RemoveClustersOfOnlyOneChild(FGeometryCollection* GeometryCollection)
+{
+	check(GeometryCollection);
+
+	bool bRemovedAny = false;
+
+	TArray<int32> DeletionList;
+	do
+	{
+		DeletionList.Reset();
+
+		for (int32 Idx = 0, Num = GeometryCollection->Transform.Num(); Idx < Num; ++Idx)
+		{
+			int32 ParentIdx = GeometryCollection->Parent[Idx];
+			if (ParentIdx != INDEX_NONE && GeometryCollection->IsClustered(Idx))
+			{
+				if (GeometryCollection->Children[Idx].Num() == 1)
+				{
+					DeletionList.Add(Idx);
+					GeometryCollectionAlgo::ParentTransforms(GeometryCollection, ParentIdx, GeometryCollection->Children[Idx].Array());
+					UpdateHierarchyLevelOfChildren(GeometryCollection, ParentIdx);
+					RecursivelyUpdateChildBoneNames(ParentIdx, GeometryCollection->Children, GeometryCollection->BoneName);
+				}
+			}
+		}
+
+		if (DeletionList.Num())
+		{
+			// Note: List is ordered by construction, so do not need to Sort()
+			FManagedArrayCollection::FProcessingParameters Params;
+			Params.bDoValidation = false; // for perf reasons
+			GeometryCollection->RemoveElements(FGeometryCollection::TransformGroup, DeletionList, Params);
+			bRemovedAny = true;
+		}
+
+		// Need to repeat until an iteration doesn't remove any nodes, to fully collapse any chains of single-child clusters
+	} while (DeletionList.Num());
+
+	return bRemovedAny;
+}
+
+bool FGeometryCollectionClusteringUtility::RemoveDanglingClusters(FGeometryCollection* GeometryCollection)
 {
 	check(GeometryCollection);
 
@@ -892,11 +933,14 @@ void FGeometryCollectionClusteringUtility::RemoveDanglingClusters(FGeometryColle
 
 	if (DeletionList.Num())
 	{
-		DeletionList.Sort();
+		// Note: List is ordered by construction, so do not need to Sort()
 		FManagedArrayCollection::FProcessingParameters Params;
 		Params.bDoValidation = false; // for perf reasons
 		GeometryCollection->RemoveElements(FGeometryCollection::TransformGroup, DeletionList, Params);
+		return true;
 	}
+
+	return false;
 }
 
 
diff --git a/Engine/Source/Runtime/Experimental/Chaos/Public/GeometryCollection/GeometryCollectionClusteringUtility.h b/Engine/Source/Runtime/Experimental/Chaos/Public/GeometryCollection/GeometryCollectionClusteringUtility.h
--- a/Engine/Source/Runtime/Experimental/Chaos/Public/GeometryCollection/GeometryCollectionClusteringUtility.h
+++ b/Engine/Source/Runtime/Experimental/Chaos/Public/GeometryCollection/GeometryCollectionClusteringUtility.h
@@ -118,8 +118,11 @@ public:
 	/** Find the lowest common ancestor index of currently selected nodes. Returns INDEX_NODE if there is no common ancestor. */
 	static int32 FindLowestCommonAncestor(FGeometryCollection* GeometryCollection, const TArray<int32>& SelectedBones);
 
-	/** Delete any cluster nodes discovered to have no children */
-	static void RemoveDanglingClusters(FGeometryCollection* GeometryCollection);
+	/** Delete any cluster nodes discovered to have no children. Returns true if any nodes were removed. */
+	static bool RemoveDanglingClusters(FGeometryCollection* GeometryCollection);
+
+	/** Simplify the geometry collection by removing cluster nodes w/ a parent and only one child, re-parenting the child to its grand-parent. Returns true if any nodes were removed. */
+	static bool RemoveClustersOfOnlyOneChild(FGeometryCollection* GeometryCollection);
 
 	static void ValidateResults(FGeometryCollection* GeometryCollection);
 
