make Validate tool always log when it fixes issues in the geometry collection

diff --git a/Engine/Plugins/Experimental/ChaosEditor/Source/FractureEditor/Private/FractureToolEditing.cpp b/Engine/Plugins/Experimental/ChaosEditor/Source/FractureEditor/Private/FractureToolEditing.cpp
--- a/Engine/Plugins/Experimental/ChaosEditor/Source/FractureEditor/Private/FractureToolEditing.cpp
+++ b/Engine/Plugins/Experimental/ChaosEditor/Source/FractureEditor/Private/FractureToolEditing.cpp
@@ -337,6 +337,7 @@ void UFractureToolValidate::Execute(TWeakPtr<FFractureEditorModeToolkit> InToolk
 					{
 						if (FGeometryCollectionClusteringUtility::RemoveClustersOfOnlyOneChild(GeometryCollection))
 						{
+							UE_LOG(LogFractureTool, Warning, TEXT("Removed one or more clusters of only one child."));
 							bDirty = true;
 						}
 					}
@@ -345,6 +346,7 @@ void UFractureToolValidate::Execute(TWeakPtr<FFractureEditorModeToolkit> InToolk
 					{
 						if (FGeometryCollectionClusteringUtility::RemoveDanglingClusters(GeometryCollection))
 						{
+							UE_LOG(LogFractureTool, Warning, TEXT("Removed one or more dangling clusters."));
 							bDirty = true;
 						}
 					}
