Simplify and improve autocluster's avoid-isolated-cluster feature:
 - Merge the two avoid-isolated-cluster checkboxes into one
 - make it always prevent isolated clusters (including the case where avoid-by-merge fails, which was not handled initially)

diff --git a/Engine/Plugins/Experimental/ChaosEditor/Source/FractureEditor/Private/FractureToolAutoCluster.cpp b/Engine/Plugins/Experimental/ChaosEditor/Source/FractureEditor/Private/FractureToolAutoCluster.cpp
--- a/Engine/Plugins/Experimental/ChaosEditor/Source/FractureEditor/Private/FractureToolAutoCluster.cpp
+++ b/Engine/Plugins/Experimental/ChaosEditor/Source/FractureEditor/Private/FractureToolAutoCluster.cpp
@@ -89,7 +89,7 @@ void UFractureToolAutoCluster::Execute(TWeakPtr<FFractureEditorModeToolkit> InTo
 					continue;
 				}
 
-				bool bNeedProximity = AutoClusterSettings->bEnforceConnectivity || AutoClusterSettings->bMergeIsolatedChildren;
+				bool bNeedProximity = AutoClusterSettings->bEnforceConnectivity || AutoClusterSettings->bAvoidIsolated;
 				if (bNeedProximity)
 				{
 					FGeometryCollectionProximityUtility ProximityUtility(GeometryCollection);
@@ -101,15 +101,16 @@ void UFractureToolAutoCluster::Execute(TWeakPtr<FFractureEditorModeToolkit> InTo
 					VoronoiPartition.SplitDisconnectedPartitions(GeometryCollection);
 				}
 
-				if (AutoClusterSettings->bMergeIsolatedChildren)
+				if (AutoClusterSettings->bAvoidIsolated)
 				{
+					// attempt to remove isolated via merging (may not succeed, as it only merges if there is a cluster in proximity)
 					VoronoiPartition.MergeSingleElementPartitions(GeometryCollection);
 				}
 
 				int32 NonEmptyPartitionCount = VoronoiPartition.GetNonEmptyPartitionCount();
 				bool bHasEmptyClusters = NonEmptyPartitionCount > 0;
 
-				if (AutoClusterSettings->bDoNotCreateSingleSiteClusters && NonEmptyPartitionCount == 1)
+				if (AutoClusterSettings->bAvoidIsolated && NonEmptyPartitionCount == 1)
 				{
 					continue;
 				}
@@ -117,12 +118,14 @@ void UFractureToolAutoCluster::Execute(TWeakPtr<FFractureEditorModeToolkit> InTo
 				int32 PartitionCount = VoronoiPartition.GetPartitionCount();
 				int32 NewClusterIndexStart = GeometryCollection->AddElements(PartitionCount, FGeometryCollection::TransformGroup);
 
-
-
 				for (int32 Index = 0; Index < PartitionCount; ++Index)
 				{
-						
 					TArray<int32> NewCluster = VoronoiPartition.GetPartition(Index);
+					if (AutoClusterSettings->bAvoidIsolated && NewCluster.Num() == 1)
+					{
+						bHasEmptyClusters = true;
+						NewCluster.Reset();
+					}
 
 					int32 NewClusterIndex = NewClusterIndexStart + Index;
 					GeometryCollection->Parent[NewClusterIndex] = ClusterIndex;
diff --git a/Engine/Plugins/Experimental/ChaosEditor/Source/FractureEditor/Private/FractureToolAutoCluster.h b/Engine/Plugins/Experimental/ChaosEditor/Source/FractureEditor/Private/FractureToolAutoCluster.h
--- a/Engine/Plugins/Experimental/ChaosEditor/Source/FractureEditor/Private/FractureToolAutoCluster.h
+++ b/Engine/Plugins/Experimental/ChaosEditor/Source/FractureEditor/Private/FractureToolAutoCluster.h
@@ -63,13 +63,9 @@ public:
 	UPROPERTY(EditAnywhere, Category = AutoCluster, meta = (DisplayName = "Enforce Cluster Connectivity"))
 	bool bEnforceConnectivity=true;
 
-	/** If true, tries to prevent the creation of clusters with only a single child, by merging such clusters into a neighbor cluster if possible */
+	/** If true, prevent the creation of clusters with only a single child. Either by merging into a neighboring cluster, or not creating the cluster. */
 	UPROPERTY(EditAnywhere, Category = AutoCluster)
-	bool bMergeIsolatedChildren = true;
-
-	/** If true, will not create a new cluster if all children would fall in that same cluster -- for example, if there is only one Voronoi site */
-	UPROPERTY(EditAnywhere, Category = AutoCluster)
-	bool bDoNotCreateSingleSiteClusters = true;
+	bool bAvoidIsolated = true;
 };
 
 
