Clarify fracture convex hull generation settings:
 - make overlap removal settings be its own group
 - rename CanExceedFraction -> Allow Larger Hull Fraction
 - rename CanRemoveFraction -> "Max Removal Fraction" and move it to the overlap removal settings group

diff --git a/Engine/Plugins/Experimental/ChaosEditor/Source/FractureEditor/Private/FractureToolConvex.cpp b/Engine/Plugins/Experimental/ChaosEditor/Source/FractureEditor/Private/FractureToolConvex.cpp
--- a/Engine/Plugins/Experimental/ChaosEditor/Source/FractureEditor/Private/FractureToolConvex.cpp
+++ b/Engine/Plugins/Experimental/ChaosEditor/Source/FractureEditor/Private/FractureToolConvex.cpp
@@ -35,27 +35,28 @@ void UFractureConvexSettings::ClearCustomConvex()
 	ConvexTool->ClearCustomConvex();
 }
 
-void UFractureConvexSettings::SaveAsDefaults()
+void UFractureConvexActions::SaveAsDefaults()
 {
+	UFractureToolConvex* ConvexTool = Cast<UFractureToolConvex>(OwnerTool.Get());
 	UFractureModeSettings* ModeSettings = GetMutableDefault<UFractureModeSettings>();
 	ModeSettings->Modify();
-	ModeSettings->ConvexFractionAllowRemove = FractionAllowRemove;
-	ModeSettings->ConvexCanExceedFraction = CanExceedFraction;
-	ModeSettings->ConvexSimplificationDistanceThreshold = SimplificationDistanceThreshold;
-	ModeSettings->ConvexRemoveOverlaps = RemoveOverlaps;
-	ModeSettings->ConvexOverlapRemovalShrinkPercent = OverlapRemovalShrinkPercent;
+	ModeSettings->ConvexFractionAllowRemove = ConvexTool->ConvexSettings->FractionAllowRemove;
+	ModeSettings->ConvexCanExceedFraction = ConvexTool->ConvexSettings->CanExceedFraction;
+	ModeSettings->ConvexSimplificationDistanceThreshold = ConvexTool->ConvexSettings->SimplificationDistanceThreshold;
+	ModeSettings->ConvexRemoveOverlaps = ConvexTool->ConvexSettings->RemoveOverlaps;
+	ModeSettings->ConvexOverlapRemovalShrinkPercent = ConvexTool->ConvexSettings->OverlapRemovalShrinkPercent;
 }
 
-void UFractureConvexSettings::SetFromDefaults()
+void UFractureConvexActions::SetFromDefaults()
 {
+	UFractureToolConvex* ConvexTool = Cast<UFractureToolConvex>(OwnerTool.Get());
 	const UFractureModeSettings* ModeSettings = GetDefault<UFractureModeSettings>();
-	FractionAllowRemove = ModeSettings->ConvexFractionAllowRemove;
-	CanExceedFraction = ModeSettings->ConvexCanExceedFraction;
-	SimplificationDistanceThreshold = ModeSettings->ConvexSimplificationDistanceThreshold;
-	RemoveOverlaps = ModeSettings->ConvexRemoveOverlaps;
-	OverlapRemovalShrinkPercent = ModeSettings->ConvexOverlapRemovalShrinkPercent;
+	ConvexTool->ConvexSettings->FractionAllowRemove = ModeSettings->ConvexFractionAllowRemove;
+	ConvexTool->ConvexSettings->CanExceedFraction = ModeSettings->ConvexCanExceedFraction;
+	ConvexTool->ConvexSettings->SimplificationDistanceThreshold = ModeSettings->ConvexSimplificationDistanceThreshold;
+	ConvexTool->ConvexSettings->RemoveOverlaps = ModeSettings->ConvexRemoveOverlaps;
+	ConvexTool->ConvexSettings->OverlapRemovalShrinkPercent = ModeSettings->ConvexOverlapRemovalShrinkPercent;
 
-	UFractureToolConvex* ConvexTool = Cast<UFractureToolConvex>(OwnerTool.Get());
 	ConvexTool->NotifyOfPropertyChangeByTool(this);
 }
 
@@ -179,6 +180,8 @@ UFractureToolConvex::UFractureToolConvex(const FObjectInitializer& ObjInit)
 {
 	ConvexSettings = NewObject<UFractureConvexSettings>(GetTransientPackage(), UFractureConvexSettings::StaticClass());
 	ConvexSettings->OwnerTool = this;
+	ConvexActions = NewObject<UFractureConvexActions>(GetTransientPackage(), UFractureConvexActions::StaticClass());
+	ConvexActions->OwnerTool = this;
 }
 
 bool UFractureToolConvex::CanExecute() const
@@ -216,6 +219,7 @@ TArray<UObject*> UFractureToolConvex::GetSettingsObjects() const
 {
 	TArray<UObject*> Settings;
 	Settings.Add(ConvexSettings);
+	Settings.Add(ConvexActions);
 	return Settings;
 }
 
diff --git a/Engine/Plugins/Experimental/ChaosEditor/Source/FractureEditor/Private/FractureToolConvex.h b/Engine/Plugins/Experimental/ChaosEditor/Source/FractureEditor/Private/FractureToolConvex.h
--- a/Engine/Plugins/Experimental/ChaosEditor/Source/FractureEditor/Private/FractureToolConvex.h
+++ b/Engine/Plugins/Experimental/ChaosEditor/Source/FractureEditor/Private/FractureToolConvex.h
@@ -22,12 +22,8 @@ public:
 		: Super(ObjInit)
 	{}
 
-	/** Fraction of the convex hulls for a transform that we can remove before using the hulls of the children */
-	UPROPERTY(EditAnywhere, Category = Automatic, meta = (DisplayName = "Can Remove Fraction", ClampMin = ".01", ClampMax = "1"))
-	double FractionAllowRemove = .5;
-
 	/** Fraction (of geometry volume) by which a cluster's convex hull volume can exceed the actual geometry volume before instead using the hulls of the children.  0 means the convex volume cannot exceed the geometry volume; 1 means the convex volume is allowed to be 100% larger (2x) the geometry volume. */
-	UPROPERTY(EditAnywhere, Category = Automatic, meta = (ClampMin = "0"))
+	UPROPERTY(EditAnywhere, Category = Automatic, meta = (DisplayName = "Allow Larger Hull Fraction", ClampMin = "0"))
 	double CanExceedFraction = .5;
 
 	/** We simplify the convex shape to keep points spaced at least this far apart (except to keep the hull from collapsing to zero volume) */
@@ -35,14 +31,18 @@ public:
 	double SimplificationDistanceThreshold = 10.0;
 
 	/** Whether to automatically cut away overlapping parts of the convex hulls, to avoid the simulation 'popping' to fix the overlaps */
-	UPROPERTY(EditAnywhere, Category = Automatic)
+	UPROPERTY(EditAnywhere, Category = AutomaticOverlapRemoval)
 	EConvexOverlapRemoval RemoveOverlaps = EConvexOverlapRemoval::All;
 
 	/** Overlap removal will be computed as if convex hulls were this percentage smaller (in range 0-100) */
-	UPROPERTY(EditAnywhere, Category = Automatic, meta = (UIMin = "0", ClampMax = "99.9"))
+	UPROPERTY(EditAnywhere, Category = AutomaticOverlapRemoval, meta = (UIMin = "0", ClampMax = "99.9"))
 	double OverlapRemovalShrinkPercent = 0.0;
 
-	/** Delete convex hulls from selected clusters.  Does not affect hulls on leaves. */
+	/** Fraction of the convex hulls for a cluster that we can remove before using the hulls of the children */
+	UPROPERTY(EditAnywhere, Category = AutomaticOverlapRemoval, meta = (DisplayName = "Max Removal Fraction", ClampMin = ".01", ClampMax = "1"))
+	double FractionAllowRemove = .5;
+
+		/** Delete convex hulls from selected clusters.  Does not affect hulls on leaves. */
 	UFUNCTION(CallInEditor, Category = Custom, meta = (DisplayName = "Delete From Selected"))
 	void DeleteFromSelected();
 
@@ -56,14 +56,32 @@ public:
 	/** Clear any manual adjustments to convex hulls on the selected bones */
 	UFUNCTION(CallInEditor, Category = Custom, meta = (DisplayName = "Clear Custom Convex"))
 	void ClearCustomConvex();
+};
+
+/**
+ * UFUNCTION actions to manage convex hulls generation for geometry collections
+ * (These are pulled out from the above settings object mainly to control their ordering in the properties panel)
+ */
+UCLASS(config = EditorPerProjectUserSettings)
+class UFractureConvexActions : public UFractureToolSettings
+{
+public:
+
+	GENERATED_BODY()
+
+	UFractureConvexActions(const FObjectInitializer& ObjInit)
+		: Super(ObjInit)
+	{}
 
 	/** Save settings as project defaults, to be used for all new geometry collections */
-	UFUNCTION(CallInEditor, Category = Defaults, meta = (DisplayName = "Save As Defaults"))
+	UFUNCTION(CallInEditor, Category = ProjectDefaults, meta = (DisplayName = "Save As Defaults"))
 	void SaveAsDefaults();
 
 	/** Set settings from current project defaults */
-	UFUNCTION(CallInEditor, Category = Defaults, meta = (DisplayName = "Set From Defaults"))
+	UFUNCTION(CallInEditor, Category = ProjectDefaults, meta = (DisplayName = "Set From Defaults"))
 	void SetFromDefaults();
+
+
 };
 
 
@@ -106,10 +124,14 @@ public:
 	virtual void Setup() override;
 
 
-protected:
-	UPROPERTY(EditAnywhere, Category = Slicing)
+	UPROPERTY(EditAnywhere, Category = Convex)
 	TObjectPtr<UFractureConvexSettings> ConvexSettings;
 
+protected:
+
+	UPROPERTY(EditAnywhere, Category = Convex)
+	TObjectPtr<UFractureConvexActions> ConvexActions;
+
 	virtual void ClearVisualizations() override
 	{
 		Super::ClearVisualizations();
diff --git a/Engine/Plugins/Experimental/ChaosEditor/Source/FractureEditor/Public/FractureModeSettings.h b/Engine/Plugins/Experimental/ChaosEditor/Source/FractureEditor/Public/FractureModeSettings.h
--- a/Engine/Plugins/Experimental/ChaosEditor/Source/FractureEditor/Public/FractureModeSettings.h
+++ b/Engine/Plugins/Experimental/ChaosEditor/Source/FractureEditor/Public/FractureModeSettings.h
@@ -32,26 +32,26 @@ public:
 
 public:
 
-	/** Default fraction of convex hulls for a transform that we can remove before using the hulls of the children */
-	UPROPERTY(config, EditAnywhere, Category = "Fracture Mode|Convex Generation Defaults", meta = (DisplayName = "Can Remove Fraction", ClampMin = ".01", ClampMax = "1"))
-	double ConvexFractionAllowRemove = .5;
-
 	/** Default fraction of geometry volume by which a cluster's convex hull volume can exceed the actual geometry volume before instead using the hulls of the children.  0 means the convex volume cannot exceed the geometry volume; 1 means the convex volume is allowed to be 100% larger (2x) the geometry volume. */
-	UPROPERTY(config, EditAnywhere, Category = "Fracture Mode|Convex Generation Defaults", meta = (ClampMin = "0"))
+	UPROPERTY(config, EditAnywhere, Category = "Fracture Mode|Convex Generation Defaults", meta = (ClampMin = "0", DisplayName = "Allow Larger Hull Fraction"))
 	float ConvexCanExceedFraction = .5;
 
 	/** Default simplification threshold for convex hulls of new geometry collections */
-	UPROPERTY(config, EditAnywhere, Category = "Fracture Mode|Convex Generation Defaults", meta = (ClampMin = "0"))
+	UPROPERTY(config, EditAnywhere, Category = "Fracture Mode|Convex Generation Defaults", meta = (ClampMin = "0", DisplayName = "Simplification Distance Threshold"))
 	float ConvexSimplificationDistanceThreshold = 10;
 
 	/** Default overlap removal setting for convex hulls of new geometry collections */
-	UPROPERTY(config, EditAnywhere, Category = "Fracture Mode|Convex Generation Defaults")
+	UPROPERTY(config, EditAnywhere, Category = "Fracture Mode|Convex Generation Defaults|Overlap Removal", meta = (DisplayName = "Remove Overlaps"))
 	EConvexOverlapRemoval ConvexRemoveOverlaps = EConvexOverlapRemoval::All;
 
 	/** Default overlap removal shrink percent (in range 0-100) for convex hulls of new geometry collections. Overlap removal will be computed assuming convex shapes will be scaled down by this percentage. */
-	UPROPERTY(config, EditAnywhere, Category = "Fracture Mode|Convex Generation Defaults", meta = (UIMin = "0", ClampMax = "99.9"))
+	UPROPERTY(config, EditAnywhere, Category = "Fracture Mode|Convex Generation Defaults|Overlap Removal", meta = (UIMin = "0", ClampMax = "99.9", DisplayName = "Overlap Removal Shrink Percent"))
 	float ConvexOverlapRemovalShrinkPercent = 0;
 
+	/** Default fraction of convex hulls for a transform that we can remove before using the hulls of the children */
+	UPROPERTY(config, EditAnywhere, Category = "Fracture Mode|Convex Generation Defaults|Overlap Removal", meta = (DisplayName = "Max Removal Fraction", ClampMin = ".01", ClampMax = "1"))
+	double ConvexFractionAllowRemove = .5;
+
 	// Apply Convex Generation Defaults to a GeometryCollection
 	void ApplyDefaultConvexSettings(FGeometryCollection& GeometryCollection) const;
 
