guard against fracture histogram crashing when calling RegenerateNodes() w/ an unset/invalid component

diff --git a/Engine/Plugins/Experimental/ChaosEditor/Source/FractureEditor/Private/SGeometryCollectionHistogram.cpp b/Engine/Plugins/Experimental/ChaosEditor/Source/FractureEditor/Private/SGeometryCollectionHistogram.cpp
--- a/Engine/Plugins/Experimental/ChaosEditor/Source/FractureEditor/Private/SGeometryCollectionHistogram.cpp
+++ b/Engine/Plugins/Experimental/ChaosEditor/Source/FractureEditor/Private/SGeometryCollectionHistogram.cpp
@@ -135,7 +135,7 @@ FGeometryCollectionHistogramItemList FGeometryCollectionHistogramItemComponent::
 	// Collect the inspected attribute 
 	FGeometryCollectionHistogramItemList NodesList;
 	
-	if (Component->GetRestCollection())
+	if (Component.IsValid() && Component->GetRestCollection())
 	{
 		FGeometryCollection* Collection = Component->GetRestCollection()->GetGeometryCollection().Get();
 
@@ -214,6 +214,10 @@ void SGeometryCollectionHistogram::SetComponents(const TArray<UGeometryCollectio
 	{
 		for (UGeometryCollectionComponent* Component : InNewComponents)
 		{
+			if (!ensure(Component))
+			{
+				continue;
+			}
 			RootNodes.Add(MakeShared<FGeometryCollectionHistogramItemComponent>(Component));
 			LeafNodes.Append(RootNodes.Last()->RegenerateNodes(LevelView));
 		}
