Chaos geometry collection: add anchored property
- add context menu to set a transform to be anchored
- add new array to clustered particle to be marked as anchored
- changed UpdateKinematicProperty to account for anchored flag

diff --git a/Engine/Plugins/Experimental/ChaosEditor/Source/FractureEditor/Private/FractureToolProperties.cpp b/Engine/Plugins/Experimental/ChaosEditor/Source/FractureEditor/Private/FractureToolProperties.cpp
--- a/Engine/Plugins/Experimental/ChaosEditor/Source/FractureEditor/Private/FractureToolProperties.cpp
+++ b/Engine/Plugins/Experimental/ChaosEditor/Source/FractureEditor/Private/FractureToolProperties.cpp
@@ -3,6 +3,8 @@
 #include "FractureToolProperties.h"
 
 #include "GeometryCollection/GeometryCollectionComponent.h"
+#include "GeometryCollection/GeometryCollectionEngineRemoval.h"
+#include "GeometryCollection/Facades/CollectionAnchoringFacade.h"
 
 
 #define LOCTEXT_NAMESPACE "FractureProperties"
@@ -55,27 +57,30 @@ void UFractureToolSetInitialDynamicState::Execute(TWeakPtr<FFractureEditorModeTo
 {
 	if (InToolkit.IsValid())
 	{
-		TSet<UGeometryCollectionComponent*> GeomCompSelection;
-		GetSelectedGeometryCollectionComponents(GeomCompSelection);
-		for (UGeometryCollectionComponent* GeometryCollectionComponent : GeomCompSelection)
+		SetSelectedInitialDynamicState(static_cast<int32>(StateSettings->InitialDynamicState));
+		InToolkit.Pin()->RefreshOutliner();
+	}
+}
+
+void UFractureToolSetInitialDynamicState::SetSelectedInitialDynamicState(int32 InitialDynamicState)
+{
+	TSet<UGeometryCollectionComponent*> GeomCompSelection;
+	GetSelectedGeometryCollectionComponents(GeomCompSelection);
+	for (UGeometryCollectionComponent* GeometryCollectionComponent : GeomCompSelection)
+	{
+		FGeometryCollectionEdit GCEdit = GeometryCollectionComponent->EditRestCollection(GeometryCollection::EEditUpdate::RestPhysics, true /*bShapeIsUnchanged*/);
+		if (UGeometryCollection* GCObject = GCEdit.GetRestCollection())
 		{
-			FGeometryCollectionEdit GCEdit = GeometryCollectionComponent->EditRestCollection(GeometryCollection::EEditUpdate::RestPhysics, true /*bShapeIsUnchanged*/);
-			if (UGeometryCollection* GCObject = GCEdit.GetRestCollection())
+			TSharedPtr<FGeometryCollection, ESPMode::ThreadSafe> GeometryCollectionPtr = GCObject->GetGeometryCollection();
+			if (FGeometryCollection* GeometryCollection = GeometryCollectionPtr.Get())
 			{
-				TSharedPtr<FGeometryCollection, ESPMode::ThreadSafe> GeometryCollectionPtr = GCObject->GetGeometryCollection();
-				if (FGeometryCollection* GeometryCollection = GeometryCollectionPtr.Get())
+				Chaos::Facades::FCollectionAnchoringFacade AnchoringFacade(*GeometryCollection);
+				if (ensure(AnchoringFacade.HasInitialDynamicStateAttribute()))
 				{
-					TManagedArray<int32>& InitialDynamicState = GeometryCollection->ModifyAttribute<int32>("InitialDynamicState", FGeometryCollection::TransformGroup);
-
-					TArray<int32> SelectedBones = GeometryCollectionComponent->GetSelectedBones();
-					for (int32 Index : SelectedBones)
-					{
-						InitialDynamicState[Index] = static_cast<int32>(StateSettings->InitialDynamicState);
-					}
+					AnchoringFacade.SetInitialDynamicState(GeometryCollectionComponent->GetSelectedBones(), static_cast<Chaos::EObjectStateType>(InitialDynamicState));
 				}
 			}
 		}
-		InToolkit.Pin()->RefreshOutliner();
 	}
 }
 
diff --git a/Engine/Plugins/Experimental/ChaosEditor/Source/FractureEditor/Private/FractureToolProperties.h b/Engine/Plugins/Experimental/ChaosEditor/Source/FractureEditor/Private/FractureToolProperties.h
--- a/Engine/Plugins/Experimental/ChaosEditor/Source/FractureEditor/Private/FractureToolProperties.h
+++ b/Engine/Plugins/Experimental/ChaosEditor/Source/FractureEditor/Private/FractureToolProperties.h
@@ -59,6 +59,7 @@ public:
 	UPROPERTY(EditAnywhere, Category = InitialDynamicState)
 	TObjectPtr<UFractureInitialDynamicStateSettings> StateSettings;
 
+	static void SetSelectedInitialDynamicState(int32 InitialDynamicState); 
 };
 
 UCLASS(config = EditorPerProjectUserSettings)
diff --git a/Engine/Plugins/Experimental/ChaosEditor/Source/FractureEditor/Private/SGeometryCollectionOutliner.cpp b/Engine/Plugins/Experimental/ChaosEditor/Source/FractureEditor/Private/SGeometryCollectionOutliner.cpp
--- a/Engine/Plugins/Experimental/ChaosEditor/Source/FractureEditor/Private/SGeometryCollectionOutliner.cpp
+++ b/Engine/Plugins/Experimental/ChaosEditor/Source/FractureEditor/Private/SGeometryCollectionOutliner.cpp
@@ -15,13 +15,15 @@
 #include "FractureEditorMode.h"
 #include "ScopedTransaction.h"
 #include "FractureSettings.h"
+#include "FractureToolProperties.h"
 #include "GeometryCollectionOutlinerDragDrop.h"
 #include "GeometryCollection/GeometryCollectionUtility.h"
+#include "GeometryCollection/Facades/CollectionAnchoringFacade.h"
 #include "PhysicsProxy/GeometryCollectionPhysicsProxy.h"
 
 #define LOCTEXT_NAMESPACE "ChaosEditor"
 
-static FText GettextFromInitialDynamicState(int32 InitialDynamicState)
+static FText GetTextFromInitialDynamicState(int32 InitialDynamicState)
 {
 	switch (InitialDynamicState)
 	{
@@ -34,6 +36,15 @@ static FText GettextFromInitialDynamicState(int32 InitialDynamicState)
 	}
 }
 
+static FText GetTextFromAnchored(bool bAnchored)
+{
+	if (bAnchored)
+	{
+		return LOCTEXT("GCOutliner_Anchored_Yes_Text", "Yes");
+	}
+	return LOCTEXT("GCOutliner_Anchored_No_Text", "-");
+}
+
 void FGeometryCollectionTreeItem::OnDragLeave(const FDragDropEvent& InDragDropEvent)
 {
 	TSharedPtr<FDragDropOperation> Operation = InDragDropEvent.GetOperation();
@@ -52,26 +63,36 @@ void FGeometryCollectionTreeItem::GenerateContextMenu(UToolMenu* Menu, SGeometry
 {
 	auto SharedOutliner = StaticCastSharedRef<SGeometryCollectionOutliner>(Outliner.AsShared());
 
-	auto MakeDynamicStateMenu = [&Outliner](UToolMenu* Menu)
-	{
-		const FName MenuEntryNames[] =
+	FToolMenuSection& StateSection = Menu->AddSection("State");
+	StateSection.AddSubMenu("FractureToolSetInitialDynamicStateMenu", NSLOCTEXT("Fracture", "FractureToolSetInitialDynamicStateMenu", "Initial Dynamic State"), FText(),
+		FNewToolMenuDelegate::CreateLambda([&Outliner](UToolMenu* Menu)
 		{
-			"NoOverride"
-			"Sleeping",
-			"Kinematic",
-			"Static"
-		};
-		const int32 MenuEntryNamesCount = sizeof(MenuEntryNames) / sizeof(FName);
-
-		FToolMenuSection& Section = Menu->AddSection("State");
-		for (int32 Index = 0; Index < MenuEntryNamesCount; ++Index)
-		{
-			Section.AddMenuEntry(MenuEntryNames[Index], GettextFromInitialDynamicState(Index), FText(), FSlateIcon(), FUIAction(FExecuteAction::CreateRaw(&Outliner, &SGeometryCollectionOutliner::SetInitialDynamicState, Index)));
-		}
-	};
+			const FName MenuEntryNames[] =
+			{
+				"NoOverride"
+				"Sleeping",
+				"Kinematic",
+				"Static"
+			};
+			const int32 MenuEntryNamesCount = sizeof(MenuEntryNames) / sizeof(FName);
+
+			FToolMenuSection& StateSection = Menu->AddSection("State");
+			for (int32 Index = 0; Index < MenuEntryNamesCount; ++Index)
+			{
+				StateSection.AddMenuEntry(MenuEntryNames[Index], GetTextFromInitialDynamicState(Index), FText(), FSlateIcon(), FUIAction(FExecuteAction::CreateRaw(&Outliner, &SGeometryCollectionOutliner::SetInitialDynamicState, Index)));
+			}
+		}));
 
-	FToolMenuSection& Section = Menu->AddSection("State");
-	Section.AddSubMenu("FractureToolSetInitialDynamicStateMenu", NSLOCTEXT("Fracture", "FractureToolSetInitialDynamicStateMenu", "Initial Dynamic State"), FText(),FNewToolMenuDelegate::CreateLambda(MakeDynamicStateMenu));
+	FToolMenuSection& AnchoredSection = Menu->AddSection("Anchored");
+	StateSection.AddSubMenu("FractureToolSetAnchoredMenu", NSLOCTEXT("Fracture", "FractureToolSetAnchoredMenu", "Anchored"), FText(),
+		FNewToolMenuDelegate::CreateLambda([&Outliner](UToolMenu* Menu)
+		{
+			FToolMenuSection& AnchoredSection = Menu->AddSection("Anchored");
+			{
+				AnchoredSection.AddMenuEntry("Yes", LOCTEXT("GCOutliner_Anchored_Yes_ContextMenuText", "Yes"), FText(), FSlateIcon(), FUIAction(FExecuteAction::CreateRaw(&Outliner, &SGeometryCollectionOutliner::SetAnchored, true)));
+				AnchoredSection.AddMenuEntry("No", LOCTEXT("GCOutliner_Anchored_No_ContextMenuText", "No"), FText(), FSlateIcon(), FUIAction(FExecuteAction::CreateRaw(&Outliner, &SGeometryCollectionOutliner::SetAnchored, false)));
+			}
+		}));
 }
 
 FColor FGeometryCollectionTreeItem::GetColorPerDepth(uint32 Depth)
@@ -185,10 +206,12 @@ TSharedRef<SWidget> SGeometryCollectionOutlinerRow::GenerateWidgetForColumn(cons
 					NameWidget.ToSharedRef()
 				];
 	}
-	if (ColumnName == SGeometryCollectionOutlinerColumnID::InitialState)
-		return Item->MakeInitialStateColumnWidget();
 	if (ColumnName == SGeometryCollectionOutlinerColumnID::RelativeSize)
 		return Item->MakeRelativeSizeColumnWidget();
+	if (ColumnName == SGeometryCollectionOutlinerColumnID::InitialState)
+		return Item->MakeInitialStateColumnWidget();
+	if (ColumnName == SGeometryCollectionOutlinerColumnID::Anchored)
+		return Item->MakeAnchoredColumnWidget();
 	if (ColumnName == SGeometryCollectionOutlinerColumnID::Damage)
 		return Item->MakeDamagesColumnWidget();
 	if (ColumnName == SGeometryCollectionOutlinerColumnID::DamageThreshold)
@@ -262,6 +285,13 @@ void SGeometryCollectionOutliner::RegenerateHeader()
 				.HAlignHeader(EHorizontalAlignment::HAlign_Center)
 				.FillWidth(CustomFillWidth)
 		);
+		HeaderRowWidget->AddColumn(
+			SHeaderRow::Column(SGeometryCollectionOutlinerColumnID::Anchored)
+				.DefaultLabel(LOCTEXT("GCOutliner_Column_Anchored", "Anchored"))
+				.DefaultTooltip(LOCTEXT("GCOutliner_Column_Anchored_ToolTip", "Anchored"))
+				.HAlignHeader(EHorizontalAlignment::HAlign_Center)
+				.FillWidth(CustomFillWidth)
+		);
 		break;
 
 	case EOutlinerColumnMode::Damage:
@@ -508,27 +538,29 @@ void SGeometryCollectionOutliner::OnSelectionChanged(FGeometryCollectionTreeItem
 
 void SGeometryCollectionOutliner::SetInitialDynamicState(int32 InDynamicState)
 {
-	FGeometryCollectionTreeItemList SelectedItems;
-	TreeView->GetSelectedItems(SelectedItems);
+	UFractureToolSetInitialDynamicState::SetSelectedInitialDynamicState(InDynamicState);
+	RegenerateItems();
+}
 
-	for (auto& SelectedItem : SelectedItems)
+void SGeometryCollectionOutliner::SetAnchored(bool bAnchored)
+{
+	// todo : eventually move this to a more central place ( tools ?)  
+	TSet<UGeometryCollectionComponent*> GeomCompSelection;
+	UFractureToolSetInitialDynamicState::GetSelectedGeometryCollectionComponents(GeomCompSelection);
+	for (UGeometryCollectionComponent* GeometryCollectionComponent : GeomCompSelection)
 	{
-		if (UGeometryCollectionComponent* Component = SelectedItem->GetComponent())
+		FGeometryCollectionEdit GCEdit = GeometryCollectionComponent->EditRestCollection(GeometryCollection::EEditUpdate::RestPhysics, true /*bShapeIsUnchanged*/);
+		if (UGeometryCollection* GCObject = GCEdit.GetRestCollection())
 		{
-			FGeometryCollectionEdit GeometryCollectionEdit = Component->EditRestCollection();
-			if (UGeometryCollection* GeometryCollectionObject = GeometryCollectionEdit.GetRestCollection())
+			TSharedPtr<FGeometryCollection, ESPMode::ThreadSafe> GeometryCollectionPtr = GCObject->GetGeometryCollection();
+			if (FGeometryCollection* GeometryCollection = GeometryCollectionPtr.Get())
 			{
-				TSharedPtr<FGeometryCollection,ESPMode::ThreadSafe> GeometryCollection = GeometryCollectionObject->GetGeometryCollection();
-				if (GeometryCollection)
+				Chaos::Facades::FCollectionAnchoringFacade AnchoringFacade(*GeometryCollection);
+				if (!AnchoringFacade.HasAnchoredAttribute())
 				{
-					TManagedArray<int32>& InitialDynamicState = GeometryCollection->ModifyAttribute<int32>("InitialDynamicState", FGeometryCollection::TransformGroup);
-					TArray<int32> SelectedBones = Component->GetSelectedBones();
-					int32 BoneIndex = SelectedItem->GetBoneIndex();
-					if (ensure(0 <= BoneIndex && BoneIndex < InitialDynamicState.Num()))
-					{
-						InitialDynamicState[SelectedItem->GetBoneIndex()] = InDynamicState;
-					}
+					AnchoringFacade.AddAnchoredAttribute();
 				}
+				AnchoringFacade.SetAnchored(GeometryCollectionComponent->GetSelectedBones(), bAnchored);
 			}
 		}
 	}
@@ -764,6 +796,7 @@ void FGeometryCollectionTreeItemBone::UpdateItemFromCollection()
 	DamageThreshold = 0;
 	Broken = false;
 	InitialState = INDEX_NONE;
+	Anchored = false;
 	RemoveOnBreakAvailable = false;
 	RemoveOnBreak = FVector4f{-1};
 	ImportedCollisionsAvailable = false;
@@ -792,7 +825,7 @@ void FGeometryCollectionTreeItemBone::UpdateItemFromCollection()
 
 			using FImplicitGeom = FGeometryDynamicCollection::FSharedImplicit;
 			const TManagedArray<FImplicitGeom>* ExternalCollisions = GeometryCollectionPtr->FindAttribute<FImplicitGeom>("ExternalCollisions", FTransformCollection::TransformGroup);
-		
+
 			if (ensure(ItemBoneIndex >= 0 && ItemBoneIndex < SimulationType.Num()))
 			{
 				const bool bHasLevelAttribute = GeometryCollectionPtr->HasAttribute("Level", FGeometryCollection::TransformGroup);
@@ -830,7 +863,13 @@ void FGeometryCollectionTreeItemBone::UpdateItemFromCollection()
 					}
 				}
 
-				InitialState = GeometryCollectionPtr->InitialDynamicState[ItemBoneIndex];
+				Chaos::Facades::FCollectionAnchoringFacade AnchoringFacade(*GeometryCollectionPtr);
+				InitialState = static_cast<int32>(AnchoringFacade.GetInitialDynamicState(ItemBoneIndex));
+				if (AnchoringFacade.HasAnchoredAttribute())
+				{
+					Anchored = AnchoringFacade.IsAnchored(ItemBoneIndex);
+				}
+				
 				if (RelativeSizes)
 				{
 					RelativeSize = (*RelativeSizes)[ItemBoneIndex];
@@ -1038,7 +1077,19 @@ TSharedRef<SWidget> FGeometryCollectionTreeItemBone::MakeInitialStateColumnWidge
 			.Padding(12.f, 0.f)
 			[
 				SNew(STextBlock)
-				.Text(GettextFromInitialDynamicState(InitialState))
+				.Text(GetTextFromInitialDynamicState(InitialState))
+				.ColorAndOpacity(ItemColor)
+			];
+}
+
+TSharedRef<SWidget> FGeometryCollectionTreeItemBone::MakeAnchoredColumnWidget() const
+{
+	return SNew(SHorizontalBox)
+		+ SHorizontalBox::Slot()
+			.Padding(12.f, 0.f)
+	[
+				SNew(STextBlock)
+				.Text(GetTextFromAnchored(Anchored))
 				.ColorAndOpacity(ItemColor)
 			];
 }
diff --git a/Engine/Plugins/Experimental/ChaosEditor/Source/FractureEditor/Private/SGeometryCollectionOutliner.h b/Engine/Plugins/Experimental/ChaosEditor/Source/FractureEditor/Private/SGeometryCollectionOutliner.h
--- a/Engine/Plugins/Experimental/ChaosEditor/Source/FractureEditor/Private/SGeometryCollectionOutliner.h
+++ b/Engine/Plugins/Experimental/ChaosEditor/Source/FractureEditor/Private/SGeometryCollectionOutliner.h
@@ -148,6 +148,7 @@ public:
 		, ItemColor(FSlateColor::UseForeground())
 		, RelativeSize(0)
 		, InitialState(INDEX_NONE)
+		, Anchored(false)
 		, Damage(0)
 		, DamageThreshold(0)
 		, Broken(false)
@@ -165,6 +166,7 @@ public:
 	TSharedRef<SWidget> MakeDamageThresholdColumnWidget() const;
 	TSharedRef<SWidget> MakeBrokenColumnWidget() const;
 	TSharedRef<SWidget> MakeInitialStateColumnWidget() const;
+	TSharedRef<SWidget> MakeAnchoredColumnWidget() const;
 	TSharedRef<SWidget> MakePostBreakTimeColumnWidget() const;
 	TSharedRef<SWidget> MakeRemovalTimeColumnWidget() const;
 	TSharedRef<SWidget> MakeImportedCollisionsColumnWidget() const;
@@ -192,6 +194,7 @@ private:
 	FText ItemText;
 	float RelativeSize;
 	int32 InitialState;
+	bool Anchored;
 	float Damage;
 	float DamageThreshold;
 	bool Broken;
@@ -210,6 +213,7 @@ namespace SGeometryCollectionOutlinerColumnID
 	// State and Size column mode
 	const FName RelativeSize("Relative Size");
 	const FName InitialState("Initial State");
+	const FName Anchored("Anchored");
 	// Damage Column Mode
 	const FName Damage("Damage");
 	const FName DamageThreshold("DamageThreshold");
@@ -273,6 +277,7 @@ public:
 	int32 GetBoneSelectionCount() const;
 
 	void SetInitialDynamicState(int32 InDynamicState);
+	void SetAnchored(bool bAnchored);
 
 	void ExpandAll();
 	void ExpandRecursive(TSharedPtr<FGeometryCollectionTreeItem> TreeItem, bool bInExpansionState) const;
diff --git a/Engine/Plugins/Experimental/ChaosEditor/Source/FractureEditor/Public/FractureTool.h b/Engine/Plugins/Experimental/ChaosEditor/Source/FractureEditor/Public/FractureTool.h
--- a/Engine/Plugins/Experimental/ChaosEditor/Source/FractureEditor/Public/FractureTool.h
+++ b/Engine/Plugins/Experimental/ChaosEditor/Source/FractureEditor/Public/FractureTool.h
@@ -76,13 +76,15 @@ public:
 
 		~FModifyContextScope();
 	};
+
+public:
+	static void GetSelectedGeometryCollectionComponents(TSet<UGeometryCollectionComponent*>& GeomCompSelection);
 	
 protected:
 	static bool IsStaticMeshSelected();
 	static bool IsGeometryCollectionSelected();
 	static void AddSingleRootNodeIfRequired(UGeometryCollection* GeometryCollectionObject);
 	static void AddAdditionalAttributesIfRequired(UGeometryCollection* GeometryCollectionObject);
-	static void GetSelectedGeometryCollectionComponents(TSet<UGeometryCollectionComponent*>& GeomCompSelection);
 	static void Refresh(FFractureToolContext& Context, FFractureEditorModeToolkit* Toolkit, bool bClearSelection=false);
 	static void SetOutlinerComponents(TArray<FFractureToolContext>& InContexts, FFractureEditorModeToolkit* Toolkit);
 	static void ClearProximity(FGeometryCollection* GeometryCollection);
diff --git a/Engine/Source/Runtime/Experimental/Chaos/Private/Chaos/PBDRigidClustering.cpp b/Engine/Source/Runtime/Experimental/Chaos/Private/Chaos/PBDRigidClustering.cpp
--- a/Engine/Source/Runtime/Experimental/Chaos/Private/Chaos/PBDRigidClustering.cpp
+++ b/Engine/Source/Runtime/Experimental/Chaos/Private/Chaos/PBDRigidClustering.cpp
@@ -103,6 +103,7 @@ namespace Chaos
 		NewParticle->SetClusterId(ClusterId(nullptr, Children.Num()));
 		NewParticle->SetClusterGroupIndex(ClusterGroupIndex);
 		NewParticle->SetStrains(0.0);
+		NewParticle->SetIsAnchored(false);
 
 		// Update clustering data structures.
 		if (MChildren.Contains(NewParticle))
@@ -207,6 +208,7 @@ namespace Chaos
 		TopLevelClusterParents.Add(NewParticle);
 		NewParticle->SetInternalCluster(true);
 		NewParticle->SetClusterId(ClusterId(nullptr, Children.Num()));
+		NewParticle->SetIsAnchored(false);
 		for (auto& Constituent : Children) MEvolution.DoInternalParticleInitilization(Constituent, NewParticle);
 
 		//
diff --git a/Engine/Source/Runtime/Experimental/Chaos/Private/Chaos/PBDRigidClusteringAlgo.cpp b/Engine/Source/Runtime/Experimental/Chaos/Private/Chaos/PBDRigidClusteringAlgo.cpp
--- a/Engine/Source/Runtime/Experimental/Chaos/Private/Chaos/PBDRigidClusteringAlgo.cpp
+++ b/Engine/Source/Runtime/Experimental/Chaos/Private/Chaos/PBDRigidClusteringAlgo.cpp
@@ -238,6 +238,7 @@ namespace Chaos
 				Chaos::FPBDRigidParticleHandle* CurrentHandle;
 				while (Queue.Dequeue(CurrentHandle) && ObjectState == EObjectStateType::Dynamic)
 				{
+					bool bIsAnchored = false;
 					if (FClusterHandle CurrentClusterHandle = CurrentHandle->CastToClustered())
 					{
 						// @question : Maybe we should just store the leaf node bodies in a
@@ -249,16 +250,25 @@ namespace Chaos
 								Queue.Enqueue(Child);
 							}
 						}
+						
+						bIsAnchored = CurrentClusterHandle->IsAnchored();
 					}
 
-					const EObjectStateType CurrState = CurrentHandle->ObjectState();
-					if (CurrState == EObjectStateType::Kinematic)
+					if (bIsAnchored)
 					{
 						ObjectState = EObjectStateType::Kinematic;
 					}
-					else if (CurrState == EObjectStateType::Static)
+					else
 					{
-						ObjectState = EObjectStateType::Static;
+						const EObjectStateType CurrState = CurrentHandle->ObjectState();
+						if (CurrState == EObjectStateType::Kinematic)
+						{
+							ObjectState = EObjectStateType::Kinematic;
+						}
+						else if (CurrState == EObjectStateType::Static)
+						{
+							ObjectState = EObjectStateType::Static;
+						}
 					}
 				}
 
diff --git a/Engine/Source/Runtime/Experimental/Chaos/Private/GeometryCollection/Facades/CollectionAnchoringFacade.cpp b/Engine/Source/Runtime/Experimental/Chaos/Private/GeometryCollection/Facades/CollectionAnchoringFacade.cpp
new file mode 100644
--- /dev/null
+++ b/Engine/Source/Runtime/Experimental/Chaos/Private/GeometryCollection/Facades/CollectionAnchoringFacade.cpp
@@ -0,0 +1,82 @@
+// Copyright Epic Games, Inc. All Rights Reserved.
+
+#include "GeometryCollection/Facades/CollectionAnchoringFacade.h"
+#include "GeometryCollection/TransformCollection.h"
+
+namespace Chaos::Facades
+{
+	FCollectionAnchoringFacade::FCollectionAnchoringFacade(FManagedArrayCollection& InCollection)
+		: InitialDynamicStateAttribute(InCollection, "InitialDynamicState", FTransformCollection::TransformGroup)
+	    , AnchoredAttribute(InCollection, "Anchored", FTransformCollection::TransformGroup)
+	{}
+
+	bool FCollectionAnchoringFacade::HasInitialDynamicStateAttribute() const
+	{
+		return InitialDynamicStateAttribute.IsValid();
+	}
+
+	bool FCollectionAnchoringFacade::HasAnchoredAttribute() const
+	{
+		return AnchoredAttribute.IsValid();
+	}
+
+	void FCollectionAnchoringFacade::AddAnchoredAttribute()
+	{
+		return AnchoredAttribute.AddAndFill(false);
+	}
+
+	void FCollectionAnchoringFacade::CopyAnchoredAttribute(const FCollectionAnchoringFacade& From)
+	{
+		if (From.HasAnchoredAttribute())
+		{
+			AnchoredAttribute.Copy(From.AnchoredAttribute);
+		}
+	}
+	
+	EObjectStateType FCollectionAnchoringFacade::GetInitialDynamicState(int32 TransformIndex) const
+	{
+		const int32 State = InitialDynamicStateAttribute.Get()[TransformIndex];
+		return static_cast<EObjectStateType>(State);
+	}
+	
+	void FCollectionAnchoringFacade::SetInitialDynamicState(int32 TransformIndex, EObjectStateType State)
+	{
+		// we don't really support sleeping as a dynamic state
+		if (ensure(State != EObjectStateType::Sleeping && State < EObjectStateType::Count))
+		{
+			InitialDynamicStateAttribute.Modify()[TransformIndex] = static_cast<int32>(State);
+		}
+	}
+
+	void FCollectionAnchoringFacade::SetInitialDynamicState(const TArray<int32>& TransformIndices, EObjectStateType State)
+	{
+		// we don't really support sleeping as a dynamic state
+		if (ensure(State != EObjectStateType::Sleeping && State < EObjectStateType::Count))
+		{
+			TManagedArray<int32>& InitialDynamicState = InitialDynamicStateAttribute.Modify();
+			for (const int32 TransformIndex: TransformIndices)
+			{
+				InitialDynamicState[TransformIndex] = static_cast<int32>(State);
+			}
+		}
+	}
+		
+	bool FCollectionAnchoringFacade::IsAnchored(int32 TransformIndex) const
+	{
+		return AnchoredAttribute.Get()[TransformIndex];
+	}
+
+	void FCollectionAnchoringFacade::SetAnchored(int32 TransformIndex, bool bValue)
+	{
+		AnchoredAttribute.Modify()[TransformIndex] = bValue;
+	}
+
+	void FCollectionAnchoringFacade::SetAnchored(const TArray<int32>& TransformIndices, bool bValue)
+	{
+		TManagedArray<bool>& Anchored = AnchoredAttribute.Modify();
+		for (const int32 TransformIndex: TransformIndices)
+		{
+			Anchored[TransformIndex] = bValue;
+		}
+	}
+}
\ No newline at end of file
diff --git a/Engine/Source/Runtime/Experimental/Chaos/Private/PhysicsProxy/GeometryCollectionPhysicsProxy.cpp b/Engine/Source/Runtime/Experimental/Chaos/Private/PhysicsProxy/GeometryCollectionPhysicsProxy.cpp
--- a/Engine/Source/Runtime/Experimental/Chaos/Private/PhysicsProxy/GeometryCollectionPhysicsProxy.cpp
+++ b/Engine/Source/Runtime/Experimental/Chaos/Private/PhysicsProxy/GeometryCollectionPhysicsProxy.cpp
@@ -31,6 +31,7 @@
 #include "Modules/ModuleManager.h"
 #include "Chaos/PullPhysicsDataImp.h"
 #include "Chaos/PBDRigidsEvolution.h"
+#include "GeometryCollection/Facades/CollectionAnchoringFacade.h"
 
 #ifndef TODO_REIMPLEMENT_INIT_COMMANDS
 #define TODO_REIMPLEMENT_INIT_COMMANDS 0
@@ -587,6 +588,11 @@ void FGeometryCollectionPhysicsProxy::Initialize(Chaos::FPBDRigidsEvolutionBase
 
 	PhysicsThreadCollection.CopyMatchingAttributesFrom(DynamicCollection, &SkipList);
 
+	// make sure we copy the anchored information over to the physics thread collection
+	const Chaos::Facades::FCollectionAnchoringFacade DynamicCollectionAnchoringFacade(DynamicCollection);
+	Chaos::Facades::FCollectionAnchoringFacade PhysicsThreadCollectionAnchoringFacade(PhysicsThreadCollection);
+	PhysicsThreadCollectionAnchoringFacade.CopyAnchoredAttribute(DynamicCollectionAnchoringFacade);
+	
 	// Copy simplicials.
 	// TODO: Ryan - Should we just transfer ownership of the SimplicialsAttribute from the DynamicCollection to
 	// the PhysicsThreadCollection?
@@ -751,6 +757,9 @@ void FGeometryCollectionPhysicsProxy::CreateNonClusteredParticles(Chaos::FPBDRig
 			RigidsSolver->GetEvolution()->CreateParticle(Handle);
 		}
 	}
+
+	// anchors
+	
 }
 
 Chaos::FPBDRigidClusteredParticleHandle* FGeometryCollectionPhysicsProxy::FindClusteredParticleHandleByItemIndex_Internal(FGeometryCollectionItemIndex ItemIndex) const
@@ -837,6 +846,8 @@ void FGeometryCollectionPhysicsProxy::InitializeBodiesPT(Chaos::FPBDRigidsSolver
 
 		const FVector WorldScale = Parameters.WorldTransform.GetScale3D();
 		const FVector::FReal MassScale = WorldScale.X * WorldScale.Y * WorldScale.Z;
+
+		const Chaos::Facades::FCollectionAnchoringFacade AnchoringFacade(PhysicsThreadCollection);
 		
 		// Iterating over the geometry group is a fast way of skipping everything that's
 		// not a leaf node, as each geometry has a transform index, which is a shortcut
@@ -867,6 +878,9 @@ void FGeometryCollectionPhysicsProxy::InitializeBodiesPT(Chaos::FPBDRigidsSolver
 					static_cast<int16>(CollisionGroup[TransformGroupIndex]),
 					CollisionParticlesPerObjectFraction);
 
+				// initialize anchoring information if available 
+				Handle->SetIsAnchored(AnchoringFacade.HasAnchoredAttribute()? AnchoringFacade.IsAnchored(TransformGroupIndex): false);
+				
 				if (Parameters.EnableClustering)
 				{
 					Handle->SetClusterGroupIndex(Parameters.ClusterGroupIndex);
diff --git a/Engine/Source/Runtime/Experimental/Chaos/Public/Chaos/PBDRigidClusteredParticles.h b/Engine/Source/Runtime/Experimental/Chaos/Public/Chaos/PBDRigidClusteredParticles.h
--- a/Engine/Source/Runtime/Experimental/Chaos/Public/Chaos/PBDRigidClusteredParticles.h
+++ b/Engine/Source/Runtime/Experimental/Chaos/Public/Chaos/PBDRigidClusteredParticles.h
@@ -88,6 +88,7 @@ class TPBDRigidClusteredParticles : public TPBDRigidParticles<T, d>
 	, MStrains(MoveTemp(Other.MStrains))
 	, MConnectivityEdges(MoveTemp(Other.MConnectivityEdges))
 	, MExternalStrains(MoveTemp(Other.MExternalStrains))
+	, MAnchored(MoveTemp(Other.MAnchored))
 	{
 		InitHelper();
 	}
@@ -125,6 +126,9 @@ class TPBDRigidClusteredParticles : public TPBDRigidParticles<T, d>
 	const auto& ConnectivityEdges(int32 Idx) const { return MConnectivityEdges[Idx]; }
 	auto& ConnectivityEdges(int32 Idx) { return MConnectivityEdges[Idx]; }
 
+	const bool& Anchored(int32 Idx) const { return MAnchored[Idx]; }
+	bool& Anchored(int32 Idx) { return MAnchored[Idx]; }
+	
 	const auto& ConnectivityEdgesArray() const { return MConnectivityEdges; }
 
 	const auto& ClusterIdsArray() const { return MClusterIds; }
@@ -142,6 +146,10 @@ class TPBDRigidClusteredParticles : public TPBDRigidParticles<T, d>
 	const auto& InternalClusterArray() const { return MInternalCluster; }
 	auto& InternalClusterArray() { return MInternalCluster; }
 
+	const auto& AnchoredArray() const { return MAnchored; }
+	auto& AnchoredArray() { return MAnchored; }
+
+	
 	typedef TPBDRigidClusteredParticleHandle<T, d> THandleType;
 	const THandleType* Handle(int32 Index) const { return static_cast<const THandleType*>(TGeometryParticles<T,d>::Handle(Index)); }
 
@@ -164,6 +172,7 @@ class TPBDRigidClusteredParticles : public TPBDRigidParticles<T, d>
 		  TArrayCollection::AddArray(&MStrains);
 		  TArrayCollection::AddArray(&MConnectivityEdges);
 	  	  TArrayCollection::AddArray(&MExternalStrains);
+	  	  TArrayCollection::AddArray(&MAnchored);
 	  }
 
 	  TArrayCollectionArray<ClusterId> MClusterIds;
@@ -179,13 +188,17 @@ class TPBDRigidClusteredParticles : public TPBDRigidParticles<T, d>
 	  TArrayCollectionArray<T> MCollisionImpulses;
 
 	  // external strains ( use by fields )
-	  // @todo(chaos) we should eventually merge MCollisionImpulses into MExternalStrains when CLustering code has been updated to not clear the impulses just before processing them 
+	  // @todo(chaos) we should eventually merge MCollisionImpulses into MExternalStrains when Clustering code has been updated to not clear the impulses just before processing them 
 	  TArrayCollectionArray<T> MExternalStrains; 
 
 	  // User set parameters
 	  TArrayCollectionArray<T> MStrains;
 
 	  TArrayCollectionArray<TArray<TConnectivityEdge<T>>> MConnectivityEdges;
+
+	  // make the particle act as a kinematic anchor,
+	  // this allows the particle to be broken off while still be anchor contributor through the connection graph   
+	  TArrayCollectionArray<bool> MAnchored;
 };
 
 using FPBDRigidClusteredParticles = TPBDRigidClusteredParticles<FReal, 3>;
diff --git a/Engine/Source/Runtime/Experimental/Chaos/Public/Chaos/ParticleHandle.h b/Engine/Source/Runtime/Experimental/Chaos/Public/Chaos/ParticleHandle.h
--- a/Engine/Source/Runtime/Experimental/Chaos/Public/Chaos/ParticleHandle.h
+++ b/Engine/Source/Runtime/Experimental/Chaos/Public/Chaos/ParticleHandle.h
@@ -1224,6 +1224,10 @@ public:
 	void SetConnectivityEdges(const TArray<TConnectivityEdge<T>>& Edges) { PBDRigidClusteredParticles->ConnectivityEdges(ParticleIdx) = Edges; }
 	void SetConnectivityEdges(TArray<TConnectivityEdge<T>>&& Edges) { PBDRigidClusteredParticles->ConnectivityEdges(ParticleIdx) = MoveTemp(Edges); }
 
+	const bool& IsAnchored() const { return PBDRigidClusteredParticles->Anchored(ParticleIdx); }
+	bool& IsAnchored() { return PBDRigidClusteredParticles->Anchored(ParticleIdx); }
+	void SetIsAnchored(const bool Value) { PBDRigidClusteredParticles->Anchored(ParticleIdx) = Value; }
+	
 	const TPBDRigidClusteredParticleHandleImp<T, d, true>* Handle() const { return PBDRigidClusteredParticles->Handle(ParticleIdx); }
 	TPBDRigidClusteredParticleHandleImp<T, d, true>* Handle() { return PBDRigidClusteredParticles->Handle(ParticleIdx); }
 
diff --git a/Engine/Source/Runtime/Experimental/Chaos/Public/GeometryCollection/Facades/CollectionAnchoringFacade.h b/Engine/Source/Runtime/Experimental/Chaos/Public/GeometryCollection/Facades/CollectionAnchoringFacade.h
new file mode 100644
--- /dev/null
+++ b/Engine/Source/Runtime/Experimental/Chaos/Public/GeometryCollection/Facades/CollectionAnchoringFacade.h
@@ -0,0 +1,42 @@
+// Copyright Epic Games, Inc. All Rights Reserved.
+
+#pragma once
+
+#include "CoreMinimal.h"
+#include "GeometryCollection/ManagedArrayAccessor.h"
+#include "GeometryCollection/ManagedArrayCollection.h"
+
+struct FManagedArrayCollection;
+
+namespace Chaos::Facades
+{
+	/**
+	 * Provides an API to define anchoring properties on a collection
+	 */
+	class CHAOS_API FCollectionAnchoringFacade
+	{
+	public:
+		FCollectionAnchoringFacade(FManagedArrayCollection& InCollection);
+
+		bool HasInitialDynamicStateAttribute() const;
+		
+		EObjectStateType GetInitialDynamicState(int32 TransformIndex) const;
+		void SetInitialDynamicState(int32 TransformIndex, EObjectStateType State);
+		void SetInitialDynamicState(const TArray<int32>& TransformIndices, EObjectStateType State);
+		
+		bool HasAnchoredAttribute() const;
+		void AddAnchoredAttribute();
+		void CopyAnchoredAttribute(const FCollectionAnchoringFacade& Other);
+
+		bool IsAnchored(int32 TransformIndex) const;
+		void SetAnchored(int32 TransformIndex, bool bValue);
+		void SetAnchored(const TArray<int32>& TransformIndices, bool bValue);
+	
+	private:
+		/** Initial dynamic state of a transform ( can be changed at runtime ) */
+		TManagedArrayAccessor<int32> InitialDynamicStateAttribute;
+
+		/** Whether a transform is anchored ( orthogonal to the initial dynamic ) )*/
+		TManagedArrayAccessor<bool> AnchoredAttribute;
+	};
+}
\ No newline at end of file
diff --git a/Engine/Source/Runtime/Experimental/Chaos/Public/GeometryCollection/ManagedArrayAccessor.h b/Engine/Source/Runtime/Experimental/Chaos/Public/GeometryCollection/ManagedArrayAccessor.h
--- a/Engine/Source/Runtime/Experimental/Chaos/Public/GeometryCollection/ManagedArrayAccessor.h
+++ b/Engine/Source/Runtime/Experimental/Chaos/Public/GeometryCollection/ManagedArrayAccessor.h
@@ -52,6 +52,12 @@ public:
 		}
 	}
 
+	/** copy from another attribute ( create if necessary ) */
+	void Copy(const TManagedArrayAccessor<T>& FromAttribute )
+	{
+		Collection.CopyAttribute(FromAttribute.Collection, Name, Group);
+	}
+
 private:
 	FManagedArrayCollection& Collection;
 	FName Name;
diff --git a/Engine/Source/Runtime/Experimental/Chaos/Public/PhysicsProxy/GeometryCollectionPhysicsProxy.h b/Engine/Source/Runtime/Experimental/Chaos/Public/PhysicsProxy/GeometryCollectionPhysicsProxy.h
--- a/Engine/Source/Runtime/Experimental/Chaos/Public/PhysicsProxy/GeometryCollectionPhysicsProxy.h
+++ b/Engine/Source/Runtime/Experimental/Chaos/Public/PhysicsProxy/GeometryCollectionPhysicsProxy.h
@@ -444,6 +444,7 @@ private:
 	TArray<FParticleHandle*> SolverClusterID;
 	TArray<FClusterHandle*> SolverClusterHandles; // make a TArray of the base clase with type
 	TArray<FClusterHandle*> SolverParticleHandles;// make a TArray of base class and join with above
+	TSet<FClusterHandle*> SolverAnchors;
 	TMap<FParticleHandle*, int32> HandleToTransformGroupIndex;
 	TMap<int32, FClusterHandle*> UniqueIdxToInternalClusterHandle;
 
diff --git a/Engine/Source/Runtime/Experimental/GeometryCollectionEngine/Private/GeometryCollection/GeometryCollectionComponent.cpp b/Engine/Source/Runtime/Experimental/GeometryCollectionEngine/Private/GeometryCollection/GeometryCollectionComponent.cpp
--- a/Engine/Source/Runtime/Experimental/GeometryCollectionEngine/Private/GeometryCollection/GeometryCollectionComponent.cpp
+++ b/Engine/Source/Runtime/Experimental/GeometryCollectionEngine/Private/GeometryCollection/GeometryCollectionComponent.cpp
@@ -45,6 +45,7 @@
 #include "Rendering/NaniteResources.h"
 #include "PrimitiveSceneInfo.h"
 #include "GeometryCollection/GeometryCollectionEngineRemoval.h"
+#include "GeometryCollection/Facades/CollectionAnchoringFacade.h"
 
 #if !(UE_BUILD_SHIPPING || UE_BUILD_TEST)
 #include "Logging/MessageLog.h"
@@ -2230,6 +2231,11 @@ void UGeometryCollectionComponent::OnCreatePhysicsState()
 					CollisionGroupArray[i] = CollisionGroup;
 				}
 			}
+
+			// let's copy anchored information if available
+			const Chaos::Facades::FCollectionAnchoringFacade RestCollectionAnchoringFacade(*RestCollection->GetGeometryCollection());
+			Chaos::Facades::FCollectionAnchoringFacade DynamicCollectionAnchoringFacade(*DynamicCollection);
+			DynamicCollectionAnchoringFacade.CopyAnchoredAttribute(RestCollectionAnchoringFacade);
 	
 			// Set up initial filter data for our particles
 			// #BGTODO We need a dummy body setup for now to allow the body instance to generate filter information. Change body instance to operate independently.
