Chaos geometryc ollection : fix Hide/unhide tool to properly refresh simulation data and make sure hidden pieces are not simulated

diff --git a/Engine/Plugins/Experimental/ChaosEditor/Source/FractureEditor/Private/FractureToolEditing.cpp b/Engine/Plugins/Experimental/ChaosEditor/Source/FractureEditor/Private/FractureToolEditing.cpp
--- a/Engine/Plugins/Experimental/ChaosEditor/Source/FractureEditor/Private/FractureToolEditing.cpp
+++ b/Engine/Plugins/Experimental/ChaosEditor/Source/FractureEditor/Private/FractureToolEditing.cpp
@@ -129,11 +129,14 @@ void UFractureToolHide::Execute(TWeakPtr<FFractureEditorModeToolkit> InToolkit)
 	{
 		FFractureEditorModeToolkit* Toolkit = InToolkit.Pin().Get();
 
+		FScopedTransaction Transaction(LOCTEXT("FractureHideTransaction", "Hide"));
+		
 		TArray<FFractureToolContext> Contexts = GetFractureToolContexts();
 
 		for (FFractureToolContext& Context : Contexts)
 		{
-			FGeometryCollection* GeometryCollection = Context.GetGeometryCollection().Get();
+			FGeometryCollectionEdit Edit(Context.GetGeometryCollectionComponent(), GeometryCollection::EEditUpdate::RestPhysicsDynamic);
+			FGeometryCollection* GeometryCollection = Edit.GetRestCollection()->GetGeometryCollection().Get();
 			UGeometryCollection* FracturedGeometryCollection = Context.GetFracturedGeometryCollection();
 
 			Context.ConvertSelectionToRigidNodes();
@@ -196,11 +199,14 @@ void UFractureToolUnhide::Execute(TWeakPtr<FFractureEditorModeToolkit> InToolkit
 	{
 		FFractureEditorModeToolkit* Toolkit = InToolkit.Pin().Get();
 
+		FScopedTransaction Transaction(LOCTEXT("FractureUnhideTransaction", "Unhide"));
+		
 		TArray<FFractureToolContext> Contexts = GetFractureToolContexts();
 
 		for (FFractureToolContext& Context : Contexts)
 		{
-			FGeometryCollection* GeometryCollection = Context.GetGeometryCollection().Get();
+			FGeometryCollectionEdit Edit(Context.GetGeometryCollectionComponent(), GeometryCollection::EEditUpdate::RestPhysicsDynamic);
+			FGeometryCollection* GeometryCollection = Edit.GetRestCollection()->GetGeometryCollection().Get();
 			UGeometryCollection* FracturedGeometryCollection = Context.GetFracturedGeometryCollection();
 
 			Context.ConvertSelectionToRigidNodes();
diff --git a/Engine/Source/Runtime/Experimental/Chaos/Private/PhysicsProxy/GeometryCollectionPhysicsProxy.cpp b/Engine/Source/Runtime/Experimental/Chaos/Private/PhysicsProxy/GeometryCollectionPhysicsProxy.cpp
--- a/Engine/Source/Runtime/Experimental/Chaos/Private/PhysicsProxy/GeometryCollectionPhysicsProxy.cpp
+++ b/Engine/Source/Runtime/Experimental/Chaos/Private/PhysicsProxy/GeometryCollectionPhysicsProxy.cpp
@@ -3116,7 +3116,10 @@ void FGeometryCollectionPhysicsProxy::InitializeSharedCollisionStructures(
 
 	for (int32 TransformGroupIndex = 0; TransformGroupIndex < NumTransforms; TransformGroupIndex++)
 	{
-		CollectionSimulatableParticles[TransformGroupIndex] = (SimulationType[TransformGroupIndex] > FGeometryCollection::ESimulationTypes::FST_None);
+		if (SimulationType[TransformGroupIndex] == FGeometryCollection::ESimulationTypes::FST_None)
+		{
+			CollectionSimulatableParticles[TransformGroupIndex] = false;
+		}
 	}
 	
 	// compute mass properties in parallel
