add icons for fracture tools: Leaf and Cluster selection and Set Remove

diff --git a/Engine/Plugins/Experimental/ChaosEditor/Source/FractureEditor/Private/FractureEditorStyle.cpp b/Engine/Plugins/Experimental/ChaosEditor/Source/FractureEditor/Private/FractureEditorStyle.cpp
--- a/Engine/Plugins/Experimental/ChaosEditor/Source/FractureEditor/Private/FractureEditorStyle.cpp
+++ b/Engine/Plugins/Experimental/ChaosEditor/Source/FractureEditor/Private/FractureEditorStyle.cpp
@@ -40,6 +40,8 @@ FFractureEditorStyle::FFractureEditorStyle()
 	Set("FractureEditor.SelectAllInLevel",     new IMAGE_BRUSH_SVG("FractureSelectLevel", IconSize));
 	Set("FractureEditor.SelectInvert",         new IMAGE_BRUSH_SVG("FractureSelectInvert", IconSize));
 	Set("FractureEditor.SelectCustom",         new IMAGE_BRUSH_SVG("FractureSelectInteractive", IconSize));
+	Set("FractureEditor.SelectLeaf",           new IMAGE_BRUSH_SVG("FractureSelectLeafNodes", IconSize));
+	Set("FractureEditor.SelectCluster",        new IMAGE_BRUSH_SVG("FractureSelectClusters", IconSize));
 
 	Set("FractureEditor.AutoCluster",          new IMAGE_BRUSH_SVG("FractureAutoCluster", IconSize));
 	Set("FractureEditor.ClusterMagnet",        new IMAGE_BRUSH_SVG("FractureMagnet", IconSize));
@@ -69,6 +71,7 @@ FFractureEditorStyle::FFractureEditorStyle()
 	Set("FractureEditor.Convex",        	   new IMAGE_BRUSH_SVG("FractureConvex", IconSize));
 	Set("FractureEditor.CustomVoronoi",        new IMAGE_BRUSH_SVG("FractureCustom", IconSize));
 	Set("FractureEditor.FixTinyGeo",	       new IMAGE_BRUSH_SVG("FractureGeoMerge", IconSize));
+	Set("FractureEditor.SetRemoveOnBreak",     new IMAGE_BRUSH_SVG("FractureSetRemoval", IconSize));
 
 	// View Settings
 	Set("FractureEditor.Exploded",             new IMAGE_BRUSH_SVG("FractureMiniExploded", LabelIconSize));
