update fracture tool names:
geomrg -> tinygeo
(new, under edit) merge -> geomrg

diff --git a/Engine/Plugins/Experimental/ChaosEditor/Source/FractureEditor/Private/FractureToolEditing.cpp b/Engine/Plugins/Experimental/ChaosEditor/Source/FractureEditor/Private/FractureToolEditing.cpp
--- a/Engine/Plugins/Experimental/ChaosEditor/Source/FractureEditor/Private/FractureToolEditing.cpp
+++ b/Engine/Plugins/Experimental/ChaosEditor/Source/FractureEditor/Private/FractureToolEditing.cpp
@@ -121,7 +121,7 @@ FSlateIcon UFractureToolMergeSelected::GetToolIcon() const
 
 void UFractureToolMergeSelected::RegisterUICommand(FFractureEditorCommands* BindingContext)
 {
-	UI_COMMAND_EXT(BindingContext, UICommandInfo, "MergeSelected", "Merge", "Merge all selected nodes into one node.", EUserInterfaceActionType::Button, FInputChord());
+	UI_COMMAND_EXT(BindingContext, UICommandInfo, "MergeSelected", "GeoMrg", "Merge all selected nodes into one node.", EUserInterfaceActionType::Button, FInputChord());
 	BindingContext->MergeSelected = UICommandInfo;
 }
 
diff --git a/Engine/Plugins/Experimental/ChaosEditor/Source/FractureEditor/Private/FractureToolFixTinyGeo.cpp b/Engine/Plugins/Experimental/ChaosEditor/Source/FractureEditor/Private/FractureToolFixTinyGeo.cpp
--- a/Engine/Plugins/Experimental/ChaosEditor/Source/FractureEditor/Private/FractureToolFixTinyGeo.cpp
+++ b/Engine/Plugins/Experimental/ChaosEditor/Source/FractureEditor/Private/FractureToolFixTinyGeo.cpp
@@ -36,7 +36,7 @@ FSlateIcon UFractureToolFixTinyGeo::GetToolIcon() const
 
 void UFractureToolFixTinyGeo::RegisterUICommand( FFractureEditorCommands* BindingContext ) 
 {
-	UI_COMMAND_EXT( BindingContext, UICommandInfo, "FixTinyGeo", "GeoMrg", "Merge pieces of geometry onto their neighbors -- use it to, for example, clean up \"too small\" pieces of geometry.", EUserInterfaceActionType::ToggleButton, FInputChord() );
+	UI_COMMAND_EXT( BindingContext, UICommandInfo, "FixTinyGeo", "TinyGeo", "Merge pieces of geometry onto their neighbors -- use it to, for example, clean up \"too small\" pieces of geometry.", EUserInterfaceActionType::ToggleButton, FInputChord() );
 	BindingContext->FixTinyGeo = UICommandInfo;
 }
 
