temporarily comment out merge/split fracture tools from fracture mode toolbar

diff --git a/Engine/Plugins/Experimental/ChaosEditor/Source/FractureEditor/Private/FractureEditorModeToolkit.cpp b/Engine/Plugins/Experimental/ChaosEditor/Source/FractureEditor/Private/FractureEditorModeToolkit.cpp
--- a/Engine/Plugins/Experimental/ChaosEditor/Source/FractureEditor/Private/FractureEditorModeToolkit.cpp
+++ b/Engine/Plugins/Experimental/ChaosEditor/Source/FractureEditor/Private/FractureEditorModeToolkit.cpp
@@ -845,8 +845,8 @@ void FFractureEditorModeToolkit::BuildToolPalette(FName PaletteIndex, class FToo
 		ToolbarBuilder.AddToolBarButton(Commands.DeleteBranch);
 		ToolbarBuilder.AddToolBarButton(Commands.Hide);
 		ToolbarBuilder.AddToolBarButton(Commands.Unhide);
-		ToolbarBuilder.AddToolBarButton(Commands.MergeSelected);
-		ToolbarBuilder.AddToolBarButton(Commands.SplitSelected);
+		//ToolbarBuilder.AddToolBarButton(Commands.MergeSelected);
+		//ToolbarBuilder.AddToolBarButton(Commands.SplitSelected);
 	}
 	else if (PaletteIndex == TEXT("Cluster"))
 	{
