re-enable geomerge tool with new icon

diff --git a/Engine/Plugins/Experimental/ChaosEditor/Source/FractureEditor/Private/FractureEditorModeToolkit.cpp b/Engine/Plugins/Experimental/ChaosEditor/Source/FractureEditor/Private/FractureEditorModeToolkit.cpp
--- a/Engine/Plugins/Experimental/ChaosEditor/Source/FractureEditor/Private/FractureEditorModeToolkit.cpp
+++ b/Engine/Plugins/Experimental/ChaosEditor/Source/FractureEditor/Private/FractureEditorModeToolkit.cpp
@@ -845,8 +845,8 @@ void FFractureEditorModeToolkit::BuildToolPalette(FName PaletteIndex, class FToo
 		ToolbarBuilder.AddToolBarButton(Commands.DeleteBranch);
 		ToolbarBuilder.AddToolBarButton(Commands.Hide);
 		ToolbarBuilder.AddToolBarButton(Commands.Unhide);
-		//ToolbarBuilder.AddToolBarButton(Commands.MergeSelected);
-		//ToolbarBuilder.AddToolBarButton(Commands.SplitSelected);
+		ToolbarBuilder.AddToolBarButton(Commands.MergeSelected);
+		//ToolbarBuilder.AddToolBarButton(Commands.SplitSelected); // Split tool intentionally disabled; prefer the 'split' option on import instead.
 	}
 	else if (PaletteIndex == TEXT("Cluster"))
 	{
diff --git a/Engine/Plugins/Experimental/ChaosEditor/Source/FractureEditor/Private/FractureEditorStyle.cpp b/Engine/Plugins/Experimental/ChaosEditor/Source/FractureEditor/Private/FractureEditorStyle.cpp
--- a/Engine/Plugins/Experimental/ChaosEditor/Source/FractureEditor/Private/FractureEditorStyle.cpp
+++ b/Engine/Plugins/Experimental/ChaosEditor/Source/FractureEditor/Private/FractureEditorStyle.cpp
@@ -61,6 +61,7 @@ FFractureEditorStyle::FFractureEditorStyle()
 	Set("FractureEditor.DeleteBranch",         new IMAGE_BRUSH_SVG("FracturePrune", IconSize));
 	Set("FractureEditor.Hide",				   new IMAGE_BRUSH_SVG("FractureHide", IconSize));
 	Set("FractureEditor.Unhide",               new IMAGE_BRUSH_SVG("FractureUnhide", IconSize));
+	Set("FractureEditor.MergeSelected",       new IMAGE_BRUSH_SVG("FractureGlue", IconSize));
 
 	Set("FractureEditor.RecomputeNormals",     new IMAGE_BRUSH_SVG("FractureNormals", IconSize));
 	Set("FractureEditor.Resample",             new IMAGE_BRUSH_SVG("FractureResample", IconSize));
