add options on geometry collection creation
 - make material duplication for internal materials optional
 - ask if user wants to split the mesh by connected component

+ implement the connected component split on the mesh description as part of the conversion, so we can do it before we've thrown out the mesh topology (geometry collection is a rendering rep that can't represent mesh topology; it splits at every UV seam and sharp edge)

diff --git a/Engine/Plugins/Experimental/ChaosEditor/Source/FractureEditor/Private/FractureToolGenerators.cpp b/Engine/Plugins/Experimental/ChaosEditor/Source/FractureEditor/Private/FractureToolGenerators.cpp
--- a/Engine/Plugins/Experimental/ChaosEditor/Source/FractureEditor/Private/FractureToolGenerators.cpp
+++ b/Engine/Plugins/Experimental/ChaosEditor/Source/FractureEditor/Private/FractureToolGenerators.cpp
@@ -5,7 +5,6 @@
 #include "Editor.h"
 #include "ScopedTransaction.h"
 #include "Engine/Selection.h"
-#include "SCreateAssetFromObject.h"
 #include "AssetToolsModule.h"
 #include "AssetRegistry/AssetRegistryModule.h"
 #include "Layers/LayersSubsystem.h"
@@ -27,8 +26,409 @@
 #include "GeometryCollection/GeometryCollectionProximityUtility.h"
 #include "FractureToolContext.h"
 
+#include "Widgets/DeclarativeSyntaxSupport.h"
+#include "Input/Reply.h"
+#include "Widgets/SCompoundWidget.h"
+
+#include "Modules/ModuleManager.h"
+#include "Widgets/SBoxPanel.h"
+#include "Widgets/Text/STextBlock.h"
+#include "Widgets/Input/SEditableTextBox.h"
+#include "Widgets/Input/SButton.h"
+#include "Styling/AppStyle.h"
+#include "GameFramework/Actor.h"
+#include "AssetRegistry/AssetData.h"
+#include "Engine/Selection.h"
+#include "Editor.h"
+#include "Misc/FileHelper.h"
+#include "IAssetTools.h"
+#include "IContentBrowserSingleton.h"
+#include "ContentBrowserModule.h"
+#include "AssetToolsModule.h"
+#include "AssetRegistry/AssetRegistryModule.h"
+#include "PackageTools.h"
+#include "SPrimaryButton.h"
+
 #define LOCTEXT_NAMESPACE "FractureToolGenerators"
 
+
+
+
+
+//////////////////////////////////////////////////////////////////////////
+// SCreateGeometryCollectionFromObject is adapted from SCreateAssetFromActor 
+
+DECLARE_DELEGATE_ThreeParams(FOnPathChosen, const FString&, bool, bool);
+
+class SCreateGeometryCollectionFromObject : public SCompoundWidget
+{
+public:
+	SLATE_BEGIN_ARGS(SCreateGeometryCollectionFromObject)
+		: _AssetFilenameSuffix()
+		, _HeadingText()
+		, _CreateButtonText()
+	{}
+
+	/** The default suffix to use for the asset filename */
+	SLATE_ARGUMENT(FString, AssetFilenameSuffix)
+
+	/** The text to display at the top of the dialog */
+	SLATE_ARGUMENT(FText, HeadingText)
+
+	/** The label for the create button */
+	SLATE_ARGUMENT(FText, CreateButtonText)
+
+	/** The path we begin naviagation at */
+	SLATE_ARGUMENT(FString, AssetPath)
+
+	SLATE_ARGUMENT(FText, DefaultNameOverride)
+
+	SLATE_ARGUMENT(bool, bSplitIslands)
+	SLATE_ARGUMENT(bool, bAddInternalMaterials)
+
+	/** Action to perform when create clicked */
+	SLATE_EVENT(FOnPathChosen, OnCreateAssetAction)
+
+	SLATE_END_ARGS()
+
+	/** Constructs this widget with InArgs */
+	void Construct(const FArguments& InArgs, TSharedPtr<SWindow> InParentWindow);
+
+private:
+
+	/** Callback when the "create asset" button is clicked. */
+	FReply OnCreateAssetFromActorClicked();
+
+	/** Callback when the selected asset path has changed. */
+	void OnSelectAssetPath(const FString& Path);
+
+	/** Destroys the window when the operation is cancelled. */
+	FReply OnCancelCreateAssetFromActor();
+
+	/** Callback when level selection has changed, will destroy the window. */
+	void OnLevelSelectionChanged(UObject* InObjectSelected);
+
+	/** Callback when the user changes the filename for the Blueprint */
+	void OnFilenameChanged(const FText& InNewName);
+
+	/** Callback to see if creating an asset is enabled */
+	bool IsCreateAssetFromActorEnabled() const;
+
+	/** Rquest to destroy the parent window */
+	void RequestDestroyParentWindow();
+
+private:
+	/** The window this widget is nested in */
+	TWeakPtr<SWindow> ParentWindow;
+
+	/** The selected path to create the asset */
+	FString AssetPath;
+
+	/** The resultant actor instance label, based on the original actor labels */
+	FString ActorInstanceLabel;
+
+	/** The default suffix to use for the asset filename */
+	FString AssetFilenameSuffix;
+
+	/** The text to display as a heading for the dialog */
+	FText HeadingText;
+
+	/** The label to be displayed on the create button */
+	FText CreateButtonText;
+
+	/** Whether to split out each connected components into separate geometry */
+	bool bSplitIslands = false;
+
+	/** Whether to duplicate materials to create internal materials when creating from a static mesh */
+	bool bAddInternalMaterials = true;
+
+	/** Filename textbox widget */
+	TSharedPtr<SEditableTextBox> FileNameWidget;
+
+	/** True if an error is currently being reported */
+	bool bIsReportingError;
+
+	/** Called when the create button is clicked */
+	FOnPathChosen OnCreateAssetAction;
+
+	/** Used to keep track of the delegate in the selection event */
+	FDelegateHandle SelectionDelegateHandle;
+};
+
+
+
+void SCreateGeometryCollectionFromObject::Construct(const FArguments& InArgs, TSharedPtr<SWindow> InParentWindow)
+{
+	AssetFilenameSuffix = InArgs._AssetFilenameSuffix;
+	HeadingText = InArgs._HeadingText;
+	CreateButtonText = InArgs._CreateButtonText;
+	OnCreateAssetAction = InArgs._OnCreateAssetAction;
+
+	if (InArgs._AssetPath.IsEmpty())
+	{
+		AssetPath = FString("/Game");
+	}
+	else
+	{
+		AssetPath = InArgs._AssetPath;
+	}
+
+	bIsReportingError = false;
+
+	FPathPickerConfig PathPickerConfig;
+	PathPickerConfig.DefaultPath = AssetPath;
+	PathPickerConfig.OnPathSelected = FOnPathSelected::CreateRaw(this, &SCreateGeometryCollectionFromObject::OnSelectAssetPath);
+
+	SelectionDelegateHandle = USelection::SelectionChangedEvent.AddSP(this, &SCreateGeometryCollectionFromObject::OnLevelSelectionChanged);
+
+	// Set up PathPickerConfig.
+	FContentBrowserModule& ContentBrowserModule = FModuleManager::LoadModuleChecked<FContentBrowserModule>("ContentBrowser");
+
+	ParentWindow = InParentWindow;
+
+	FString PackageName;
+	ActorInstanceLabel.Empty();
+
+	if (InArgs._DefaultNameOverride.IsEmpty())
+	{
+		USelection* SelectedActors = GEditor->GetSelectedActors();
+		for (FSelectionIterator Iter(*SelectedActors); Iter; ++Iter)
+		{
+			AActor* Actor = Cast<AActor>(*Iter);
+			if (Actor)
+			{
+				ActorInstanceLabel += Actor->GetActorLabel();
+				ActorInstanceLabel += TEXT("_");
+				break;
+			}
+		}
+	}
+	else
+	{
+		ActorInstanceLabel = InArgs._DefaultNameOverride.ToString();
+	}
+
+	ActorInstanceLabel = UPackageTools::SanitizePackageName(ActorInstanceLabel + AssetFilenameSuffix);
+
+	FString AssetName = ActorInstanceLabel;
+	FString BasePath = AssetPath / AssetName;
+	FAssetToolsModule& AssetToolsModule = FModuleManager::LoadModuleChecked<FAssetToolsModule>("AssetTools");
+	AssetToolsModule.Get().CreateUniqueAssetName(BasePath, TEXT(""), PackageName, AssetName);
+
+	ChildSlot
+		[
+			SNew(SVerticalBox)
+			+ SVerticalBox::Slot()
+		.FillHeight(1.0f)
+		[
+			ContentBrowserModule.Get().CreatePathPicker(PathPickerConfig)
+		]
+	+ SVerticalBox::Slot()
+		.AutoHeight()
+		[
+			SNew(SBorder)
+			.BorderImage(FAppStyle::Get().GetBrush("Brushes.Panel"))
+		.Padding(16.f)
+		[
+			SNew(SVerticalBox)
+
+			+ SVerticalBox::Slot()
+		.AutoHeight()
+		[
+			SNew(SHorizontalBox)
+			+ SHorizontalBox::Slot()
+		.Padding(0.0f, 0.0f, 8.0f, 0.0f)
+		.AutoWidth()
+		.VAlign(VAlign_Center)
+		[
+			SNew(STextBlock)
+			.Text(HeadingText)
+		]
+
+	+ SHorizontalBox::Slot()
+		.FillWidth(1.0f)
+		[
+			SAssignNew(FileNameWidget, SEditableTextBox)
+			.Text(FText::FromString(AssetName))
+		.OnTextChanged(this, &SCreateGeometryCollectionFromObject::OnFilenameChanged)
+		]
+		]
+	+ SVerticalBox::Slot()
+		.AutoHeight()
+		[
+			SNew(SBorder)
+			.BorderImage(FAppStyle::Get().GetBrush("Brushes.Panel"))
+			.Padding(4.f)
+			[
+				SNew(SVerticalBox)
+				+ SVerticalBox::Slot()
+				[
+					SNew(SHorizontalBox)
+					+ SHorizontalBox::Slot()
+					.Padding(0.0f, 0.0f, 8.0f, 0.0f)
+					.FillWidth(0.5f)
+					.VAlign(VAlign_Center)
+					[
+						SNew(STextBlock)
+						.Text(LOCTEXT("AddInternalsLabel", "Add Internal Materials"))
+						.ToolTipText(LOCTEXT("InternalMatsToolTip", "If checked, static mesh materials will be duplicated to create materials for internal surfaces. Does not apply if creating from a Geometry Collection source."))
+					]
+
+				+ SHorizontalBox::Slot()
+					.FillWidth(1.0f)
+					[
+						SNew(SCheckBox)
+						.IsChecked_Lambda([this]()->ECheckBoxState
+							{
+								return bAddInternalMaterials ? ECheckBoxState::Checked : ECheckBoxState::Unchecked;
+							})
+						.OnCheckStateChanged_Lambda([this](ECheckBoxState NewState)
+							{
+								if (NewState == ECheckBoxState::Undetermined)
+								{
+									return;
+								}
+								bAddInternalMaterials = NewState == ECheckBoxState::Checked;
+							})
+						.ToolTipText(LOCTEXT("InternalMatsToolTip", "If checked, static mesh materials will be duplicated to create materials for internal surfaces. Does not apply if creating from a Geometry Collection source."))
+					]
+				]
+				+ SVerticalBox::Slot()
+				[
+					SNew(SHorizontalBox)
+					+ SHorizontalBox::Slot()
+					.Padding(0.0f, 0.0f, 8.0f, 0.0f)
+					.FillWidth(0.5f)
+					.VAlign(VAlign_Center)
+					[
+						SNew(STextBlock)
+						.Text(LOCTEXT("SplitComponentsLabel", "Split Components"))
+						.ToolTipText(LOCTEXT("SplitComponentsToolTip", "If checked, connected components will be split into separate bones on import."))
+					]
+
+				+ SHorizontalBox::Slot()
+					.FillWidth(1.0f)
+					[
+						SNew(SCheckBox)
+						.IsChecked_Lambda([this]()->ECheckBoxState
+							{
+								return bSplitIslands ? ECheckBoxState::Checked : ECheckBoxState::Unchecked;
+							})
+						.OnCheckStateChanged_Lambda([this](ECheckBoxState NewState)
+						{
+							if (NewState == ECheckBoxState::Undetermined)
+							{
+								return;
+							}
+							bSplitIslands = NewState == ECheckBoxState::Checked;
+						})
+						.ToolTipText(LOCTEXT("SplitComponentsToolTip", "If checked, connected components will be split into separate bones on import."))
+					]
+				]
+			]
+		]
+
+	+ SVerticalBox::Slot()
+		.HAlign(HAlign_Right)
+		.AutoHeight()
+		.Padding(0.f, 16.f, 0.f, 0.f)
+		[
+			SNew(SHorizontalBox)
+			+ SHorizontalBox::Slot()
+		.AutoWidth()
+		[
+			SNew(SPrimaryButton)
+			.OnClicked(this, &SCreateGeometryCollectionFromObject::OnCreateAssetFromActorClicked)
+		.Text(CreateButtonText)
+		]
+
+	+ SHorizontalBox::Slot()
+		.AutoWidth()
+		.Padding(16.f, 0.f, 0.f, 0.f)
+		[
+			SNew(SButton)
+			.Text(LOCTEXT("CancelButtonText", "Cancel"))
+		.OnClicked(this, &SCreateGeometryCollectionFromObject::OnCancelCreateAssetFromActor)
+		]
+		]
+		]
+		]
+		];
+
+	OnFilenameChanged(FText::FromString(AssetName));
+}
+
+void SCreateGeometryCollectionFromObject::RequestDestroyParentWindow()
+{
+	USelection::SelectionChangedEvent.Remove(SelectionDelegateHandle);
+
+	if (ParentWindow.IsValid())
+	{
+		ParentWindow.Pin()->RequestDestroyWindow();
+	}
+}
+
+FReply SCreateGeometryCollectionFromObject::OnCreateAssetFromActorClicked()
+{
+	RequestDestroyParentWindow();
+	OnCreateAssetAction.ExecuteIfBound(AssetPath / FileNameWidget->GetText().ToString(), bAddInternalMaterials, bSplitIslands);
+	return FReply::Handled();
+}
+
+FReply SCreateGeometryCollectionFromObject::OnCancelCreateAssetFromActor()
+{
+	RequestDestroyParentWindow();
+	return FReply::Handled();
+}
+
+void SCreateGeometryCollectionFromObject::OnSelectAssetPath(const FString& Path)
+{
+	AssetPath = Path;
+	OnFilenameChanged(FileNameWidget->GetText());
+}
+
+void SCreateGeometryCollectionFromObject::OnLevelSelectionChanged(UObject* InObjectSelected)
+{
+	// When actor selection changes, this window should be destroyed.
+	RequestDestroyParentWindow();
+}
+
+void SCreateGeometryCollectionFromObject::OnFilenameChanged(const FText& InNewName)
+{
+	TArray<FAssetData> AssetData;
+	FAssetRegistryModule& AssetRegistryModule = FModuleManager::LoadModuleChecked<FAssetRegistryModule>("AssetRegistry");
+	AssetRegistryModule.Get().GetAssetsByPath(FName(*AssetPath), AssetData);
+
+	FText ErrorText;
+	if (!FFileHelper::IsFilenameValidForSaving(InNewName.ToString(), ErrorText) || !FName(*InNewName.ToString()).IsValidObjectName(ErrorText))
+	{
+		FileNameWidget->SetError(ErrorText);
+		bIsReportingError = true;
+		return;
+	}
+	else
+	{
+		// Check to see if the name conflicts
+		for (auto Iter = AssetData.CreateConstIterator(); Iter; ++Iter)
+		{
+			if (Iter->AssetName.ToString() == InNewName.ToString())
+			{
+				FileNameWidget->SetError(LOCTEXT("AssetInUseError", "Asset name already in use!"));
+				bIsReportingError = true;
+				return;
+			}
+		}
+	}
+
+	FileNameWidget->SetError(FText::FromString(TEXT("")));
+	bIsReportingError = false;
+}
+
+bool SCreateGeometryCollectionFromObject::IsCreateAssetFromActorEnabled() const
+{
+	return !bIsReportingError;
+}
+
 // Creates an undo/redo action that (un)registers and object with the Asset Registry.
 // Upon undo this causes the object to be unregistered and as a result be removed from
 // Content Browsers.
@@ -103,9 +503,9 @@ void UFractureToolGenerateAsset::OpenGenerateAssetDialog(TArray<AActor*>& Actors
 
 	// NOTE - the parent window has to completely exist before this one does so the parent gets set properly.
 	// This is why we do not just put this in the Contents()[ ... ] of the Window above.
-	TSharedPtr<SCreateAssetFromObject> CreateAssetDialog;
+	TSharedPtr<SCreateGeometryCollectionFromObject> CreateAssetDialog;
 	PickAssetPathWindow->SetContent(
-		SAssignNew(CreateAssetDialog, SCreateAssetFromObject, PickAssetPathWindow)
+		SAssignNew(CreateAssetDialog, SCreateGeometryCollectionFromObject, PickAssetPathWindow)
 		.AssetFilenameSuffix(TEXT("GeometryCollection"))
 		.HeadingText(LOCTEXT("CreateGeometryCollection_Heading", "Geometry Collection Name"))
 		.CreateButtonText(LOCTEXT("CreateGeometryCollection_ButtonLabel", "Create Geometry Collection"))
@@ -125,7 +525,7 @@ void UFractureToolGenerateAsset::OpenGenerateAssetDialog(TArray<AActor*>& Actors
 
 }
 
-void UFractureToolGenerateAsset::OnGenerateAssetPathChosen(const FString& InAssetPath, TArray<AActor*> Actors)
+void UFractureToolGenerateAsset::OnGenerateAssetPathChosen(const FString& InAssetPath, bool bAddInternalMaterials, bool bSplitComponents, TArray<AActor*> Actors)
 {	
 	FScopedTransaction Transaction(LOCTEXT("GenerateAsset", "Generate Geometry Collection Asset"));
 
@@ -144,7 +544,7 @@ void UFractureToolGenerateAsset::OnGenerateAssetPathChosen(const FString& InAsse
 
 		AGeometryCollectionActor* GeometryCollectionActor = Cast<AGeometryCollectionActor>(FirstActor);
 		
-		GeometryCollectionActor = ConvertActorsToGeometryCollection(InAssetPath, Actors);
+		GeometryCollectionActor = ConvertActorsToGeometryCollection(InAssetPath, bAddInternalMaterials, bSplitComponents, Actors);
 
 		GeometryCollectionComponent = GeometryCollectionActor->GetGeometryCollectionComponent();
 
@@ -186,7 +586,7 @@ void UFractureToolGenerateAsset::OnGenerateAssetPathChosen(const FString& InAsse
 	}
 }
 
-AGeometryCollectionActor* UFractureToolGenerateAsset::ConvertActorsToGeometryCollection(const FString& InAssetPath, TArray<AActor*>& Actors)
+AGeometryCollectionActor* UFractureToolGenerateAsset::ConvertActorsToGeometryCollection(const FString& InAssetPath, bool bAddInternalMaterials, bool bSplitComponents, TArray<AActor*>& Actors)
 {
 	ensure(Actors.Num() > 0);
 	AActor* FirstActor = Actors[0];
@@ -226,9 +626,9 @@ AGeometryCollectionActor* UFractureToolGenerateAsset::ConvertActorsToGeometryCol
 				// Record the contributing source on the asset.
 				FSoftObjectPath SourceSoftObjectPath(ComponentStaticMesh);
 				decltype(FGeometryCollectionSource::SourceMaterial) SourceMaterials(StaticMeshComponent->GetMaterials());
-				FracturedGeometryCollection->GeometrySource.Add({ SourceSoftObjectPath, ComponentTransform, SourceMaterials });
+				FracturedGeometryCollection->GeometrySource.Add({ SourceSoftObjectPath, ComponentTransform, SourceMaterials, bAddInternalMaterials, bSplitComponents });
 
-				FGeometryCollectionEngineConversion::AppendStaticMesh(ComponentStaticMesh, SourceMaterials, ComponentTransform, FracturedGeometryCollection, false);
+				FGeometryCollectionEngineConversion::AppendStaticMesh(ComponentStaticMesh, SourceMaterials, ComponentTransform, FracturedGeometryCollection, false, bAddInternalMaterials, bSplitComponents);
 			}
 		}
 
@@ -445,7 +845,7 @@ void UFractureToolResetAsset::Execute(TWeakPtr<FFractureEditorModeToolkit> InToo
 					const UObject* SourceMesh = Source.SourceGeometryObject.TryLoad();
 					if (const UStaticMesh* SourceStaticMesh = Cast<UStaticMesh>(SourceMesh))
 					{
-						FGeometryCollectionEngineConversion::AppendStaticMesh(SourceStaticMesh, Source.SourceMaterial, Source.LocalTransform, GeometryCollectionObject, false);
+						FGeometryCollectionEngineConversion::AppendStaticMesh(SourceStaticMesh, Source.SourceMaterial, Source.LocalTransform, GeometryCollectionObject, false, Source.bAddInternalMaterials, Source.bSplitComponents);
 					}
 					else if (const USkeletalMesh* SourceSkeletalMesh = Cast<USkeletalMesh>(SourceMesh))
 					{
diff --git a/Engine/Plugins/Experimental/ChaosEditor/Source/FractureEditor/Private/FractureToolGenerators.h b/Engine/Plugins/Experimental/ChaosEditor/Source/FractureEditor/Private/FractureToolGenerators.h
--- a/Engine/Plugins/Experimental/ChaosEditor/Source/FractureEditor/Private/FractureToolGenerators.h
+++ b/Engine/Plugins/Experimental/ChaosEditor/Source/FractureEditor/Private/FractureToolGenerators.h
@@ -30,8 +30,8 @@ public:
 
 private:
 	void OpenGenerateAssetDialog(TArray<AActor*>& Actors);
-	void OnGenerateAssetPathChosen(const FString& InAssetPath, TArray<AActor*> Actors);
-	AGeometryCollectionActor* ConvertActorsToGeometryCollection(const FString& InAssetPath, TArray<AActor*>& Actors);
+	void OnGenerateAssetPathChosen(const FString& InAssetPath, bool bAddInternalMaterials, bool bSplitComponents, TArray<AActor*> Actors);
+	AGeometryCollectionActor* ConvertActorsToGeometryCollection(const FString& InAssetPath, bool bAddInternalMaterials, bool bSplitComponents, TArray<AActor*>& Actors);
 	class AGeometryCollectionActor* CreateNewGeometryActor(const FString& Name, const FTransform& Transform, bool AddMaterials /*= false*/);
 	static ULevel* GetSelectedLevel();
 	static AActor* AddActor(ULevel* InLevel, UClass* Class);
diff --git a/Engine/Source/Runtime/Experimental/GeometryCollectionEngine/GeometryCollectionEngine.Build.cs b/Engine/Source/Runtime/Experimental/GeometryCollectionEngine/GeometryCollectionEngine.Build.cs
--- a/Engine/Source/Runtime/Experimental/GeometryCollectionEngine/GeometryCollectionEngine.Build.cs
+++ b/Engine/Source/Runtime/Experimental/GeometryCollectionEngine/GeometryCollectionEngine.Build.cs
@@ -30,6 +30,13 @@ namespace UnrealBuildTool.Rules
 					"StaticMeshDescription",
 				}
 				);
+			PrivateDependencyModuleNames.AddRange(
+				new string[]
+				{
+					"MeshConversion",
+					"GeometryCore",
+				}
+				);
 
 			PrivateIncludePathModuleNames.Add("DerivedDataCache");
 
diff --git a/Engine/Source/Runtime/Experimental/GeometryCollectionEngine/Private/GeometryCollection/GeometryCollectionEngineConversion.cpp b/Engine/Source/Runtime/Experimental/GeometryCollectionEngine/Private/GeometryCollection/GeometryCollectionEngineConversion.cpp
--- a/Engine/Source/Runtime/Experimental/GeometryCollectionEngine/Private/GeometryCollection/GeometryCollectionEngineConversion.cpp
+++ b/Engine/Source/Runtime/Experimental/GeometryCollectionEngine/Private/GeometryCollection/GeometryCollectionEngineConversion.cpp
@@ -25,6 +25,8 @@
 #include "StaticMeshOperations.h"
 #include "Physics/Experimental/ChaosInterfaceUtils.h"
 #include "PhysicsEngine/BodySetup.h"
+#include "MeshDescriptionBuilder.h"
+#include "VertexConnectedComponents.h"
 
 
 DEFINE_LOG_CATEGORY_STATIC(UGeometryCollectionConversionLogging, Log, All);
@@ -86,7 +88,7 @@ static FVector GetMeshBuildScale3D(const UStaticMesh& StaticMesh)
 
 void FGeometryCollectionEngineConversion::AppendMeshDescription(
 	const FMeshDescription* MeshDescription, const FString& Name, int32 MaterialStartIndex, const FTransform& StaticMeshTransform, 
-	FGeometryCollection* GeometryCollection, UBodySetup* BodySetup, bool ReindexMaterials)
+	FGeometryCollection* GeometryCollection, UBodySetup* BodySetup, bool ReindexMaterials, bool bAddInternalMaterials)
 {
 #if WITH_EDITORONLY_DATA
 
@@ -213,8 +215,9 @@ void FGeometryCollectionEngineConversion::AppendMeshDescription(
 
 		TargetVisible[TargetIndex] = true;
 
-		// Materials are ganged in pairs and we want the id to associate with the first of each pair.
-		TargetMaterialID[TargetIndex] = MaterialStartIndex + (MeshDescription->GetTrianglePolygonGroup(TriangleIndex) * 2);
+		// If adding internal materials, then materials are ganged in pairs and we want the id to associate with the first of each pair.
+		int32 MaterialIndexScale = 1 + int32(bAddInternalMaterials);
+		TargetMaterialID[TargetIndex] = MaterialStartIndex + (MeshDescription->GetTrianglePolygonGroup(TriangleIndex) * MaterialIndexScale);
 
 		// Is this right?
 		TargetMaterialIndex[TargetIndex] = TargetIndex;
@@ -411,7 +414,9 @@ int32 FGeometryCollectionEngineConversion::AppendMaterials(const TArray<UMateria
 	return MaterialStart;
 }
 
-bool FGeometryCollectionEngineConversion::AppendStaticMesh(const UStaticMesh* StaticMesh, const TArray<UMaterialInterface*>& Materials, const FTransform& StaticMeshTransform, UGeometryCollection* GeometryCollectionObject, bool bReindexMaterials)
+bool FGeometryCollectionEngineConversion::AppendStaticMesh(const UStaticMesh* StaticMesh, const TArray<UMaterialInterface*>& Materials, 
+	const FTransform& StaticMeshTransform, UGeometryCollection* GeometryCollectionObject, bool bReindexMaterials,
+	bool bAddInternalMaterials, bool bSplitComponents)
 {
 #if WITH_EDITORONLY_DATA
 
@@ -422,9 +427,9 @@ bool FGeometryCollectionEngineConversion::AppendStaticMesh(const UStaticMesh* St
 	FGeometryCollection* GeometryCollection = GeometryCollectionPtr.Get();
 	check(GeometryCollection);
 
-	if (AppendStaticMesh(StaticMesh, StartMaterialIndex, StaticMeshTransform, GeometryCollection, bReindexMaterials))
+	if (AppendStaticMesh(StaticMesh, StartMaterialIndex, StaticMeshTransform, GeometryCollection, bReindexMaterials, bAddInternalMaterials, bSplitComponents))
 	{
-		AppendMaterials(Materials, GeometryCollectionObject, true);
+		AppendMaterials(Materials, GeometryCollectionObject, bAddInternalMaterials);
 		return true;
 	}
 
@@ -432,7 +437,8 @@ bool FGeometryCollectionEngineConversion::AppendStaticMesh(const UStaticMesh* St
 	return false;
 }
 
-bool FGeometryCollectionEngineConversion::AppendStaticMesh(const UStaticMesh* StaticMesh, int32 StartMaterialIndex, const FTransform& StaticMeshTransform, FGeometryCollection* GeometryCollection, bool bReindexMaterials)
+bool FGeometryCollectionEngineConversion::AppendStaticMesh(const UStaticMesh* StaticMesh, int32 StartMaterialIndex,  const FTransform& StaticMeshTransform,
+	FGeometryCollection* GeometryCollection, bool bReindexMaterials, bool bAddInternalMaterials, bool bSplitComponents)
 {
 #if WITH_EDITORONLY_DATA
 	if (StaticMesh)
@@ -449,8 +455,134 @@ bool FGeometryCollectionEngineConversion::AppendStaticMesh(const UStaticMesh* St
 					StaticMeshTransform.GetTranslation(),
 					StaticMeshTransform.GetScale3D() * MeshBuildScale3D
 			);
-			
-			AppendMeshDescription(MeshDescription, StaticMesh->GetName(), StartMaterialIndex, MeshTransform, GeometryCollection, StaticMesh->GetBodySetup(), bReindexMaterials);
+
+			if (bSplitComponents)
+			{
+				int32 MaxVID = MeshDescription->Vertices().Num();
+				UE::Geometry::FVertexConnectedComponents Components(MaxVID);
+				for (const FTriangleID TriangleID : MeshDescription->Triangles().GetElementIDs())
+				{
+					TArrayView<const FVertexID> TriangleIDs = MeshDescription->GetTriangleVertices(TriangleID);
+					Components.ConnectVertices(TriangleIDs[0].GetValue(), TriangleIDs[1].GetValue());
+					Components.ConnectVertices(TriangleIDs[1].GetValue(), TriangleIDs[2].GetValue());
+				}
+				if (Components.HasMultipleComponents(MaxVID, 2))
+				{
+					// look up vertex positions
+					TVertexAttributesConstRef<FVector3f> VertexPositions = MeshDescription->GetVertexPositions();
+
+					// vertex instance attributes
+					FStaticMeshConstAttributes Attributes(*MeshDescription);
+					TVertexInstanceAttributesConstRef<FVector2f> InstanceUVs = Attributes.GetVertexInstanceUVs();
+					TVertexInstanceAttributesConstRef<FVector3f> InstanceNormals = Attributes.GetVertexInstanceNormals();
+					TVertexInstanceAttributesConstRef<FVector3f> InstanceTangents = Attributes.GetVertexInstanceTangents();
+					TVertexInstanceAttributesConstRef<float> InstanceBiTangentSign = Attributes.GetVertexInstanceBinormalSigns();
+					TVertexInstanceAttributesConstRef<FVector4f> InstanceColors = Attributes.GetVertexInstanceColors();
+					const int NumUVLayers = InstanceUVs.GetNumChannels();
+
+					TMap<int32, int32> Map = Components.MakeComponentMap(MaxVID, 2);
+					int32 NumIslands = Map.Num();
+					
+					TArray<FMeshDescription> Descriptions;
+					Descriptions.SetNum(NumIslands);
+					TArray<FMeshDescriptionBuilder> Builders;
+					Builders.SetNum(NumIslands);
+					for (int32 MeshIdx = 0; MeshIdx < NumIslands; ++MeshIdx)
+					{
+						FStaticMeshAttributes MeshAttributes(Descriptions[MeshIdx]);
+						MeshAttributes.Register();
+
+						Builders[MeshIdx].SetMeshDescription(&Descriptions[MeshIdx]);
+						Builders[MeshIdx].SuspendMeshDescriptionIndexing();
+						Builders[MeshIdx].SetNumUVLayers(NumUVLayers);
+					}
+					for (TPair<int32, int32> IDToIdx : Map)
+					{
+						int32 ID = IDToIdx.Key;
+						int32 Idx = IDToIdx.Value;
+						int32 NumVertices = Components.GetComponentSize(ID);
+						Builders[Idx].ReserveNewVertices(NumVertices);
+					}
+					TArray<int32> VertexIDMap;
+					VertexIDMap.Init(INDEX_NONE, MeshDescription->Vertices().Num());
+
+					for (const FVertexID VertexID : MeshDescription->Vertices().GetElementIDs())
+					{
+						int32 MeshID = Components.GetComponent(VertexID.GetValue());
+						int32* MeshIdx = Map.Find(MeshID);
+						if (MeshIdx)
+						{
+							FVector Position = (FVector)VertexPositions.Get(VertexID);
+							VertexIDMap[VertexID.GetValue()] = Builders[*MeshIdx].AppendVertex(Position);
+						}
+					}
+					for (const FTriangleID TriangleID : MeshDescription->Triangles().GetElementIDs())
+					{
+						TArrayView<const FVertexID> TriangleVerts = MeshDescription->GetTriangleVertices(TriangleID);
+						TArrayView<const FVertexInstanceID> SourceInstanceTri = MeshDescription->GetTriangleVertexInstances(TriangleID);
+						int32 MeshID = Components.GetComponent(TriangleVerts[0].GetValue());
+						int32 MeshIdx = Map[MeshID];
+						FMeshDescriptionBuilder& Builder = Builders[MeshIdx];
+
+						// create new vtx instances for each triangle
+						FVertexInstanceID DestInstanceTri[3];
+						for (int32 j = 0; j < 3; ++j)
+						{
+							const FVertexID TriVertex = VertexIDMap[TriangleVerts[j].GetValue()];
+							DestInstanceTri[j] = Builder.AppendInstance(TriVertex);
+						}
+						// add the triangle to MeshDescription
+						FPolygonGroupID MaterialID = MeshDescription->GetTrianglePolygonGroup(TriangleID);
+						FTriangleID NewTriangleID = Builder.AppendTriangle(DestInstanceTri[0], DestInstanceTri[1], DestInstanceTri[2], MaterialID);
+						// transfer UVs.  Note the Builder sets both the shared and per-instance UVs from this
+						for (int32 UVLayer = 0; UVLayer < NumUVLayers; ++UVLayer)
+						{
+							FUVID UVIDs[3] = { FUVID(-1), FUVID(-1), FUVID(-1) };
+							for (int32 j = 0; j < 3; ++j)
+							{
+								FVector2D UV = (FVector2D)InstanceUVs.Get(SourceInstanceTri[j], UVLayer);
+								UVIDs[j] = Builder.AppendUV(UV, UVLayer);
+							}
+
+							// append the UV triangle - builder takes care of the rest
+							Builder.AppendUVTriangle(NewTriangleID, UVIDs[0], UVIDs[1], UVIDs[2], UVLayer);
+						}
+
+						// Set instance attributes: normal/tangent/bitangent frame and color
+						for (int32 j = 0; j < 3; ++j)
+						{
+							const FVertexInstanceID SourceInstanceID = SourceInstanceTri[j];
+							const FVertexInstanceID DestInstanceID = DestInstanceTri[j];
+							FVector TriVertNormal = (FVector)InstanceNormals.Get(SourceInstanceID);
+							FVector TriVertTangent = (FVector)InstanceTangents.Get(SourceInstanceID);
+							float BiTangentSign = (float)InstanceBiTangentSign.Get(SourceInstanceID);
+							Builder.SetInstanceTangentSpace(DestInstanceID, TriVertNormal, TriVertTangent, BiTangentSign);
+							FVector4f InstColor = InstanceColors.Get(SourceInstanceID);
+							Builder.SetInstanceColor(DestInstanceID, InstColor);
+						}
+					}
+
+					for (int32 MeshIdx = 0; MeshIdx < NumIslands; ++MeshIdx)
+					{
+						Builders[MeshIdx].ResumeMeshDescriptionIndexing();
+					}
+
+					for (FMeshDescription& MD : Descriptions)
+					{
+						AppendMeshDescription(&MD, StaticMesh->GetName(), StartMaterialIndex, MeshTransform, GeometryCollection, StaticMesh->GetBodySetup(), false, bAddInternalMaterials);
+					}
+
+					if (bReindexMaterials)
+					{
+						GeometryCollection->ReindexMaterials();
+					}
+
+					return true;
+				}
+				// else only one component -- fall back to just using the original mesh description
+			}
+
+			AppendMeshDescription(MeshDescription, StaticMesh->GetName(), StartMaterialIndex, MeshTransform, GeometryCollection, StaticMesh->GetBodySetup(), bReindexMaterials, bAddInternalMaterials);
 			return true;
 		}
 	}
@@ -710,7 +842,8 @@ void FGeometryCollectionEngineConversion::AppendGeometryCollection(const UGeomet
 }
 
 
-void FGeometryCollectionEngineConversion::AppendStaticMesh(const UStaticMesh* StaticMesh, const UStaticMeshComponent* StaticMeshComponent, const FTransform& StaticMeshTransform, UGeometryCollection* GeometryCollectionObject, bool ReindexMaterials)
+void FGeometryCollectionEngineConversion::AppendStaticMesh(const UStaticMesh* StaticMesh, const UStaticMeshComponent* StaticMeshComponent, const FTransform& StaticMeshTransform, UGeometryCollection* GeometryCollectionObject, 
+	bool ReindexMaterials, bool bAddInternalMaterials, bool bSplitComponents)
 {
 	if (StaticMesh == nullptr)
 	{
@@ -731,7 +864,7 @@ void FGeometryCollectionEngineConversion::AppendStaticMesh(const UStaticMesh* St
 	GeometryCollectionObject->Materials.Remove(BoneSelectedMaterial);
 	Materials.Remove(BoneSelectedMaterial);
 	
-	AppendStaticMesh(StaticMesh, Materials, StaticMeshTransform, GeometryCollectionObject, ReindexMaterials);
+	AppendStaticMesh(StaticMesh, Materials, StaticMeshTransform, GeometryCollectionObject, ReindexMaterials, bAddInternalMaterials, bSplitComponents);
 }
 
 
diff --git a/Engine/Source/Runtime/Experimental/GeometryCollectionEngine/Public/GeometryCollection/GeometryCollectionEngineConversion.h b/Engine/Source/Runtime/Experimental/GeometryCollectionEngine/Public/GeometryCollection/GeometryCollectionEngineConversion.h
--- a/Engine/Source/Runtime/Experimental/GeometryCollectionEngine/Public/GeometryCollection/GeometryCollectionEngineConversion.h
+++ b/Engine/Source/Runtime/Experimental/GeometryCollectionEngine/Public/GeometryCollection/GeometryCollectionEngineConversion.h
@@ -41,7 +41,8 @@ public:
 	 * @param BodySetup : Optional collision setup to transfer to the geometry collection
 	 * @param bReindexMaterials	: Whether to reindex materials -- if appending multiple meshes, pass false and call ReindexMaterials afterwards
 	 */
-	static void AppendMeshDescription(const FMeshDescription* MeshDescription, const FString& Name, int32 StartMaterialIndex, const FTransform& StaticMeshTransform, FGeometryCollection* GeometryCollection, UBodySetup* BodySetup = nullptr, bool bReindexMaterials = true);
+	static void AppendMeshDescription(const FMeshDescription* MeshDescription, const FString& Name, int32 StartMaterialIndex, const FTransform& StaticMeshTransform,
+		FGeometryCollection* GeometryCollection, UBodySetup* BodySetup = nullptr, bool bReindexMaterials = true, bool bAddInternalMaterials = true);
 
 	/**
 	 * Get a HiRes (or LOD 0 if no HiRes available) MeshDescription for the given static mesh, and make sure it includes normals and tangents.
@@ -58,7 +59,8 @@ public:
 	 * @param bReindexMaterials	: Whether to reindex materials -- if appending multiple meshes, pass false and call ReindexMaterials afterwards
 	 * @return true if succeeded in appending a mesh
 	 */
-	static bool AppendStaticMesh(const UStaticMesh* StaticMesh, const TArray<UMaterialInterface*>& Materials, const FTransform& StaticMeshTransform, UGeometryCollection* GeometryCollectionObject, bool bReindexMaterials = true);
+	static bool AppendStaticMesh(const UStaticMesh* StaticMesh, const TArray<UMaterialInterface*>& Materials, const FTransform& StaticMeshTransform, 
+		UGeometryCollection* GeometryCollectionObject, bool bReindexMaterials = true, bool bAddInternalMaterials = true, bool bSplitComponents = false);
 
 	// TODO: consider having the below fn return bool, and have a helper that returns the StartMaterialIndex from UGeometryCollection
 	//  before any materials are actually added, so we can call the AppendMaterial function after, and not need to use its return value ...
@@ -74,7 +76,8 @@ public:
 	 * @param bReindexMaterials	: Whether to reindex materials -- if appending multiple meshes, pass false and call ReindexMaterials afterwards
 	 * @return true if succeeded in appending a mesh
 	 */
-	static bool AppendStaticMesh(const UStaticMesh* StaticMesh, int32 StartMaterialIndex, const FTransform& StaticMeshTransform, FGeometryCollection* GeometryCollection, bool bReindexMaterials = true);
+	static bool AppendStaticMesh(const UStaticMesh* StaticMesh, int32 StartMaterialIndex, const FTransform& StaticMeshTransform, 
+		FGeometryCollection* GeometryCollection, bool bReindexMaterials = true, bool bAddInternalMaterials = true, bool bSplitComponents = false);
 
 	/**
 	*  Appends a static mesh to a GeometryCollectionComponent.
@@ -82,7 +85,8 @@ public:
 	*  @param StaticMeshTransform : Mesh transform.
 	*  @param GeometryCollection  : Collection to append the mesh into.
 	*/
-	static void AppendStaticMesh(const UStaticMesh* StaticMesh, const UStaticMeshComponent *StaticMeshComponent, const FTransform& StaticMeshTransform, UGeometryCollection* GeometryCollection, bool bReindexMaterials = true);
+	static void AppendStaticMesh(const UStaticMesh* StaticMesh, const UStaticMeshComponent *StaticMeshComponent, const FTransform& StaticMeshTransform, UGeometryCollection* GeometryCollection,
+		bool bReindexMaterials = true, bool bAddInternalMaterials = true, bool bSplitComponents = false);
 
 	/**
 	*  Appends a skeletal mesh to a GeometryCollectionComponent.
diff --git a/Engine/Source/Runtime/Experimental/GeometryCollectionEngine/Public/GeometryCollection/GeometryCollectionObject.h b/Engine/Source/Runtime/Experimental/GeometryCollectionEngine/Public/GeometryCollection/GeometryCollectionObject.h
--- a/Engine/Source/Runtime/Experimental/GeometryCollectionEngine/Public/GeometryCollection/GeometryCollectionObject.h
+++ b/Engine/Source/Runtime/Experimental/GeometryCollectionEngine/Public/GeometryCollection/GeometryCollectionObject.h
@@ -33,6 +33,14 @@ struct GEOMETRYCOLLECTIONENGINE_API FGeometryCollectionSource
 
 	UPROPERTY(BlueprintReadWrite, EditAnywhere, Category = "GeometrySource")
 	TArray<TObjectPtr<UMaterialInterface>> SourceMaterial;
+
+	/** Whether source materials should be duplicated to create slots for internal materials. Does not apply if the source is a GeometryCollection. */
+	UPROPERTY(BlueprintReadWrite, EditAnywhere, Category = "GeometrySource")
+	bool bAddInternalMaterials = true;
+
+	/** Whether connected components in the source object should be added as separate pieces of geometry. */
+	UPROPERTY(BlueprintReadWrite, EditAnywhere, Category = "GeometrySource")
+	bool bSplitComponents = false;
 };
 
 USTRUCT(BlueprintType)
diff --git a/Engine/Source/Runtime/GeometryCore/Public/VertexConnectedComponents.h b/Engine/Source/Runtime/GeometryCore/Public/VertexConnectedComponents.h
--- a/Engine/Source/Runtime/GeometryCore/Public/VertexConnectedComponents.h
+++ b/Engine/Source/Runtime/GeometryCore/Public/VertexConnectedComponents.h
@@ -237,6 +237,28 @@ public:
 		return false;
 	}
 
+	bool HasMultipleComponents(int32 MaxVID, int32 KeepSizeThreshold = 0)
+	{
+		int32 FoundComponent = -1;
+		for (int32 VID = 0; VID < MaxVID; VID++)
+		{
+			int32 SetID = DisjointSet.Find(VID);
+			if (KeepSizeThreshold > 0 && DisjointSet.Sizes[SetID] < KeepSizeThreshold)
+			{
+				continue;
+			}
+			if (FoundComponent == -1)
+			{
+				FoundComponent = SetID;
+			}
+			else if (FoundComponent != SetID)
+			{
+				return true;
+			}
+		}
+		return false;
+	}
+
 	// Map arbitrary set IDs to indices from 0 to k-1 (if there are k components)
 	template<typename TriangleMeshType>
 	TMap<int32, int32> MakeComponentMap(const TriangleMeshType& Mesh, int32 KeepSizeThreshold = 0)
@@ -262,6 +284,25 @@ public:
 		return ComponentMap;
 	}
 
+	TMap<int32, int32> MakeComponentMap(int32 MaxVID, int32 KeepSizeThreshold = 0)
+	{
+		TMap<int32, int32> ComponentMap;
+		int32 CurrentIdx = 0;
+		for (int32 VID = 0; VID < MaxVID; VID++)
+		{
+			int32 SetID = DisjointSet.Find(VID);
+			if (KeepSizeThreshold > 0 && DisjointSet.Sizes[SetID] < KeepSizeThreshold)
+			{
+				continue;
+			}
+			if (!ComponentMap.Contains(SetID))
+			{
+				ComponentMap.Add(SetID, CurrentIdx++);
+			}
+		}
+		return ComponentMap;
+	}
+
 	inline int32 GetComponent(int32 VertexID)
 	{
 		return DisjointSet.Find(VertexID);
