temporariily re-hide fracture geomerge tool

diff --git a/Engine/Plugins/Experimental/ChaosEditor/Source/FractureEditor/Private/FractureEditorModeToolkit.cpp b/Engine/Plugins/Experimental/ChaosEditor/Source/FractureEditor/Private/FractureEditorModeToolkit.cpp
--- a/Engine/Plugins/Experimental/ChaosEditor/Source/FractureEditor/Private/FractureEditorModeToolkit.cpp
+++ b/Engine/Plugins/Experimental/ChaosEditor/Source/FractureEditor/Private/FractureEditorModeToolkit.cpp
@@ -845,7 +845,7 @@ void FFractureEditorModeToolkit::BuildToolPalette(FName PaletteIndex, class FToo
 		ToolbarBuilder.AddToolBarButton(Commands.DeleteBranch);
 		ToolbarBuilder.AddToolBarButton(Commands.Hide);
 		ToolbarBuilder.AddToolBarButton(Commands.Unhide);
-		ToolbarBuilder.AddToolBarButton(Commands.MergeSelected);
+		//ToolbarBuilder.AddToolBarButton(Commands.MergeSelected);
 		//ToolbarBuilder.AddToolBarButton(Commands.SplitSelected); // Split tool intentionally disabled; prefer the 'split' option on import instead.
 	}
 	else if (PaletteIndex == TEXT("Cluster"))
