Fix fracture mesh tool leaving large cut hierarchy when applying multiple cuts

diff --git a/Engine/Plugins/Experimental/ChaosEditor/Source/FractureEditor/Private/FractureToolMeshCut.cpp b/Engine/Plugins/Experimental/ChaosEditor/Source/FractureEditor/Private/FractureToolMeshCut.cpp
--- a/Engine/Plugins/Experimental/ChaosEditor/Source/FractureEditor/Private/FractureToolMeshCut.cpp
+++ b/Engine/Plugins/Experimental/ChaosEditor/Source/FractureEditor/Private/FractureToolMeshCut.cpp
@@ -227,6 +227,8 @@ public:
 
 		ResultGeometryIndex = -1;
 		const float ProgressFrac = 1.0 / MeshTransforms.Num();
+
+		int32 OriginalNumTransforms = CollectionCopy->NumElements(FGeometryCollection::TransformGroup);
 		
 		TArray<int32> BonesToCut = Selection;
 		for (const FTransform& ScatterTransform : MeshTransforms)
@@ -238,6 +240,11 @@ public:
 
 			FProgressCancel::FProgressScope ProgressScope(Progress, ProgressFrac);
 			int32 Index = CutWithMesh(CuttingMesh, ScatterTransform, InternalSurfaceMaterials, *CollectionCopy, BonesToCut, PointSpacing, Transform);
+			if (Progress && Progress->Cancelled())
+			{
+				return;
+			}
+
 			int32 NewLen = Algo::RemoveIf(BonesToCut, [&](int32 Bone)
 				{
 					return !CollectionCopy->IsVisible(Bone); // remove already-fractured pieces from the to-cut list
@@ -258,6 +265,33 @@ public:
 			}
 		}
 
+		if (ResultGeometryIndex > -1)
+		{
+			TArray<int32> ToRemove;
+			for (int32 NewIdx = OriginalNumTransforms; NewIdx < CollectionCopy->NumElements(FGeometryCollection::TransformGroup); ++NewIdx)
+			{
+				if (CollectionCopy->IsRigid(NewIdx))
+				{
+					int32 ParentIdx = CollectionCopy->Parent[NewIdx];
+					if (ParentIdx >= OriginalNumTransforms)
+					{
+						do
+						{
+							ParentIdx = CollectionCopy->Parent[ParentIdx];
+						} while (CollectionCopy->Parent[ParentIdx] >= OriginalNumTransforms);
+						CollectionCopy->ParentTransforms(ParentIdx, NewIdx);
+					}
+				}
+				else
+				{
+					ToRemove.Add(NewIdx);
+				}
+			}
+			FManagedArrayCollection::FProcessingParameters ProcessingParams;
+			ProcessingParams.bDoValidation = false;
+			CollectionCopy->RemoveElements(FGeometryCollection::TransformGroup, ToRemove, ProcessingParams);
+		}
+
 		if (Progress && Progress->Cancelled())
 		{
 			return;
