Change default cluster fraction for fracture autocluster tool from .5 to .25

diff --git a/Engine/Plugins/Experimental/ChaosEditor/Source/FractureEditor/Private/FractureToolAutoCluster.h b/Engine/Plugins/Experimental/ChaosEditor/Source/FractureEditor/Private/FractureToolAutoCluster.h
--- a/Engine/Plugins/Experimental/ChaosEditor/Source/FractureEditor/Private/FractureToolAutoCluster.h
+++ b/Engine/Plugins/Experimental/ChaosEditor/Source/FractureEditor/Private/FractureToolAutoCluster.h
@@ -57,7 +57,7 @@ public:
 
 	/** Choose the number of Voronoi sites used for clustering as a fraction of the number of child bones to process */
 	UPROPERTY(EditAnywhere, Category = ClusterSize, meta = (DisplayName = "Cluster Fraction", ClampMin = "0", ClampMax = ".5", EditCondition = "ClusterSizeMethod == EClusterSizeMethod::ByFractionOfInput"))
-	float SiteCountFraction = .5;
+	float SiteCountFraction = .25;
 
 	/** If true, bones will only be added to the same cluster if they are physically connected (either directly, or via other bones in the same cluster) */
 	UPROPERTY(EditAnywhere, Category = AutoCluster, meta = (DisplayName = "Enforce Cluster Connectivity"))
