remove duplicated if statement from GeometryCollectionOutlinerDragDrop.cpp

diff --git a/Engine/Plugins/Experimental/ChaosEditor/Source/FractureEditor/Private/GeometryCollectionOutlinerDragDrop.cpp b/Engine/Plugins/Experimental/ChaosEditor/Source/FractureEditor/Private/GeometryCollectionOutlinerDragDrop.cpp
--- a/Engine/Plugins/Experimental/ChaosEditor/Source/FractureEditor/Private/GeometryCollectionOutlinerDragDrop.cpp
+++ b/Engine/Plugins/Experimental/ChaosEditor/Source/FractureEditor/Private/GeometryCollectionOutlinerDragDrop.cpp
@@ -39,16 +39,9 @@ bool FGeometryCollectionBoneDragDrop::ValidateDrop(const FGeometryCollection* Ot
 			return false;
 		}
 
-		if (ContainsInstance())
-		{
-			MessageText = NSLOCTEXT("GeometryCollectionOutliner", "EmbeddedToCluster_GeometryCollectionOutliner",
-				"Cannot parent instanced embedded geometry to cluster.");
-			return false;
-		}
-
 		if (ContainsEmbedded())
 		{
-			MessageText = NSLOCTEXT("GeometryCollectionOutliner", "EmbeddedToCluster_GeometryCollectionOutliner2",
+			MessageText = NSLOCTEXT("GeometryCollectionOutliner", "EmbeddedToCluster_GeometryCollectionOutliner",
 				"Embedded geometry will convert to rigids when parented to cluster.");
 			return true;
 		}
