Fix fracture mode capturing 'escape' when in PIE and simulate modes, making it difficult to leave the mode

diff --git a/Engine/Plugins/Experimental/ChaosEditor/Source/FractureEditor/Private/FractureEditorModeToolkit.cpp b/Engine/Plugins/Experimental/ChaosEditor/Source/FractureEditor/Private/FractureEditorModeToolkit.cpp
--- a/Engine/Plugins/Experimental/ChaosEditor/Source/FractureEditor/Private/FractureEditorModeToolkit.cpp
+++ b/Engine/Plugins/Experimental/ChaosEditor/Source/FractureEditor/Private/FractureEditorModeToolkit.cpp
@@ -921,6 +921,11 @@ void FFractureEditorModeToolkit::BindCommands()
 			{
 				GEditor->SelectNone(true, true, false);
 			}
+		}),
+		FCanExecuteAction::CreateLambda([this]()
+		{
+			// don't capture escape when in PIE or simulating
+			return GEditor->PlayWorld == NULL && !GEditor->bIsSimulatingInEditor;
 		})
 	);
 
