make fracture autouv tool use more appropriate texture types when baking

diff --git a/Engine/Plugins/Experimental/ChaosEditor/Source/FractureEditor/Private/FractureToolUV.cpp b/Engine/Plugins/Experimental/ChaosEditor/Source/FractureEditor/Private/FractureToolUV.cpp
--- a/Engine/Plugins/Experimental/ChaosEditor/Source/FractureEditor/Private/FractureToolUV.cpp
+++ b/Engine/Plugins/Experimental/ChaosEditor/Source/FractureEditor/Private/FractureToolUV.cpp
@@ -319,6 +319,12 @@ bool UFractureToolAutoUV::SaveGeneratedTexture(UE::Geometry::TImageBuilder<FVect
 
 	FString NewAssetPath = FPaths::Combine(PackageFolderPath, ObjectBaseName);
 
+	FTexture2DBuilder::ETextureType TextureType = FTexture2DBuilder::ETextureType::ColorLinear;
+	if (AutoUVSettings->BakeTextureType == ETextureType::Normals)
+	{
+		TextureType = FTexture2DBuilder::ETextureType::NormalMap;
+	}
+
 	UTexture2D* GeneratedTexture = nullptr;
 	bool bNeedsNewPackage = true;
 	if (bAllowReplace)
@@ -331,11 +337,11 @@ bool UFractureToolAutoUV::SaveGeneratedTexture(UE::Geometry::TImageBuilder<FVect
 			if (UTexture2D* OldTexture = Cast<UTexture2D>(OldObject))
 			{
 				FTexture2DBuilder TextureBuilder;
-				TextureBuilder.InitializeAndReplaceExistingTexture(OldTexture, FTexture2DBuilder::ETextureType::Color, ImageBuilder.GetDimensions());
+				TextureBuilder.InitializeAndReplaceExistingTexture(OldTexture, TextureType, ImageBuilder.GetDimensions());
 				TextureBuilder.Copy(ImageBuilder);
 				TextureBuilder.Commit(false);
 				GeneratedTexture = TextureBuilder.GetTexture2D();
-				FTexture2DBuilder::CopyPlatformDataToSourceData(GeneratedTexture, FTexture2DBuilder::ETextureType::Color);
+				FTexture2DBuilder::CopyPlatformDataToSourceData(GeneratedTexture, TextureType);
 				bNeedsNewPackage = false;
 			}
 			else // old asset was wrong type; delete to replace
@@ -348,11 +354,11 @@ bool UFractureToolAutoUV::SaveGeneratedTexture(UE::Geometry::TImageBuilder<FVect
 	if (!GeneratedTexture)
 	{
 		FTexture2DBuilder TextureBuilder;
-		TextureBuilder.Initialize(FTexture2DBuilder::ETextureType::Color, ImageBuilder.GetDimensions());
+		TextureBuilder.Initialize(TextureType, ImageBuilder.GetDimensions());
 		TextureBuilder.Copy(ImageBuilder);
 		TextureBuilder.Commit(false);
 		GeneratedTexture = TextureBuilder.GetTexture2D();
-		FTexture2DBuilder::CopyPlatformDataToSourceData(GeneratedTexture, FTexture2DBuilder::ETextureType::Color);
+		FTexture2DBuilder::CopyPlatformDataToSourceData(GeneratedTexture, TextureType);
 	}
 	AutoUVSettings->Result = GeneratedTexture;
 	check(GeneratedTexture);
