stop fracture editor auto-selecting components, because the editor assumes that if components are selected then only one actor can be selected

diff --git a/Engine/Plugins/Experimental/ChaosEditor/Source/FractureEditor/Private/FractureEditorModeToolkit.cpp b/Engine/Plugins/Experimental/ChaosEditor/Source/FractureEditor/Private/FractureEditorModeToolkit.cpp
--- a/Engine/Plugins/Experimental/ChaosEditor/Source/FractureEditor/Private/FractureEditorModeToolkit.cpp
+++ b/Engine/Plugins/Experimental/ChaosEditor/Source/FractureEditor/Private/FractureEditorModeToolkit.cpp
@@ -1492,11 +1492,6 @@ void FFractureEditorModeToolkit::SetOutlinerComponents(const TArray<UGeometryCol
 
 void FFractureEditorModeToolkit::SetBoneSelection(UGeometryCollectionComponent* InRootComponent, const TArray<int32>& InSelectedBones, bool bClearCurrentSelection, int32 FocusBoneIdx)
 {
-	if (InSelectedBones.Num() > 0 && !InRootComponent->IsSelected())
-	{
-		GEditor->SelectComponent(InRootComponent, true, true);
-	}
-
 	OutlinerView->SetBoneSelection(InRootComponent, InSelectedBones, bClearCurrentSelection, FocusBoneIdx);
 	HistogramView->SetBoneSelection(InRootComponent, InSelectedBones, bClearCurrentSelection, FocusBoneIdx);
 
