make bone colors always recompute on actor selection, to work around stale bone colors causing bones to appear selected when they are not.
 + fix FScopedColorEdit::Sanitize() clearing the bUpdated flag when nothing was sanitized.

diff --git a/Engine/Plugins/Experimental/ChaosEditor/Source/FractureEditor/Private/FractureEditorMode.cpp b/Engine/Plugins/Experimental/ChaosEditor/Source/FractureEditor/Private/FractureEditorMode.cpp
--- a/Engine/Plugins/Experimental/ChaosEditor/Source/FractureEditor/Private/FractureEditorMode.cpp
+++ b/Engine/Plugins/Experimental/ChaosEditor/Source/FractureEditor/Private/FractureEditorMode.cpp
@@ -504,7 +504,8 @@ void UFractureEditorMode::OnActorSelectionChanged(const TArray<UObject*>& NewSel
 			}
 			::GeometryCollection::GenerateTemporaryGuids(RestCollectionEdit.GetRestCollection()->GetGeometryCollection().Get());
 			
-			FScopedColorEdit ShowBoneColorsEdit(GeometryCollectionComponent);
+			constexpr bool bForceUpdate = true; // Force the bone selection and highlight to refresh so bone colors reflect the selection
+			FScopedColorEdit ShowBoneColorsEdit(GeometryCollectionComponent, bForceUpdate);
 			ShowBoneColorsEdit.SetEnableBoneSelection(true);
 			// ShowBoneColorsEdit.SetLevelViewMode(ViewLevel);
 			ShowBoneColorsEdit.Sanitize(); // Clean any stale data (e.g. due to the geometry being edited via a different component)
diff --git a/Engine/Source/Runtime/Experimental/GeometryCollectionEngine/Private/GeometryCollection/GeometryCollectionComponent.cpp b/Engine/Source/Runtime/Experimental/GeometryCollectionEngine/Private/GeometryCollection/GeometryCollectionComponent.cpp
--- a/Engine/Source/Runtime/Experimental/GeometryCollectionEngine/Private/GeometryCollection/GeometryCollectionComponent.cpp
+++ b/Engine/Source/Runtime/Experimental/GeometryCollectionEngine/Private/GeometryCollection/GeometryCollectionComponent.cpp
@@ -2838,7 +2838,7 @@ void FScopedColorEdit::Sanitize()
 			const int32 NumHighlightRemoved = Component->HighlightedBones.RemoveAll([this, NumTransforms](int32 Index) {
 				return Index < 0 || Index >= NumTransforms;
 				});
-			bUpdated = NumSelectionRemoved || NumHighlightRemoved;
+			bUpdated = bUpdated || NumSelectionRemoved || NumHighlightRemoved;
 		}
 	}
 }
