change 'level up' button in fracture mode to operate on clusters as well as rigids (but filter embedded -> rigid parent, as it's not valid for embedded to level up)

diff --git a/Engine/Plugins/Experimental/ChaosEditor/Source/FractureEditor/Private/FractureToolClustering.cpp b/Engine/Plugins/Experimental/ChaosEditor/Source/FractureEditor/Private/FractureToolClustering.cpp
--- a/Engine/Plugins/Experimental/ChaosEditor/Source/FractureEditor/Private/FractureToolClustering.cpp
+++ b/Engine/Plugins/Experimental/ChaosEditor/Source/FractureEditor/Private/FractureToolClustering.cpp
@@ -250,7 +250,9 @@ void UFractureToolMoveUp::Execute(TWeakPtr<FFractureEditorModeToolkit> InToolkit
 		{
 			FGeometryCollectionEdit Edit(Context.GetGeometryCollectionComponent(), GeometryCollection::EEditUpdate::RestPhysicsDynamic);
 
-			Context.ConvertSelectionToRigidNodes();
+			Context.ConvertEmbeddedSelectionToParents(); // it's not valid to 'move up' embedded geo
+			Context.RemoveChildrenOfSelectedNodesFromSelection(); // if a node's children are also selected, just move up the parent node
+			
 			FGeometryCollectionClusteringUtility::MoveUpOneHierarchyLevel(Context.GetGeometryCollection().Get(), Context.GetSelection());
 			FGeometryCollectionClusteringUtility::RemoveDanglingClusters(Context.GetGeometryCollection().Get());
 			Refresh(Context, Toolkit, true);
diff --git a/Engine/Plugins/Experimental/ChaosEditor/Source/FractureEditor/Private/FractureToolContext.cpp b/Engine/Plugins/Experimental/ChaosEditor/Source/FractureEditor/Private/FractureToolContext.cpp
--- a/Engine/Plugins/Experimental/ChaosEditor/Source/FractureEditor/Private/FractureToolContext.cpp
+++ b/Engine/Plugins/Experimental/ChaosEditor/Source/FractureEditor/Private/FractureToolContext.cpp
@@ -5,6 +5,8 @@
 #include "GeometryCollection/GeometryCollectionClusteringUtility.h"
 #include "GeometryCollection/GeometryCollectionUtility.h"
 
+#include "Algo/RemoveIf.h"
+
 
 FFractureToolContext::FFractureToolContext(UGeometryCollectionComponent* InGeometryCollectionComponent)
 {
@@ -116,6 +118,51 @@ void FFractureToolContext::RandomReduceSelection(float ProbToKeep)
 	}
 }
 
+void FFractureToolContext::RemoveChildrenOfSelectedNodesFromSelection()
+{
+	Sanitize();
+	const TManagedArray<int32>& SimulationType = GeometryCollection->SimulationType;
+	const TManagedArray<int32>& Parents = GeometryCollection->Parent;
+	TSet<int32> SelectionSet(SelectedBones);
+	SelectedBones.SetNum(Algo::RemoveIf(SelectedBones, [&](const int32& Bone)
+	{
+		int32 Parent = Bone;
+		while ((Parent = Parents[Parent]) != INDEX_NONE)
+		{
+			if (SelectionSet.Contains(Parent))
+			{
+				return true;
+			}
+		}
+		return false;
+	}), false);
+}
+
+void FFractureToolContext::ConvertEmbeddedSelectionToParents()
+{
+	Sanitize();
+	const TManagedArray<int32>& SimulationType = GeometryCollection->SimulationType;
+	const TManagedArray<int32>& Parents = GeometryCollection->Parent;
+
+	for (int32 SelBoneIdx = 0; SelBoneIdx < SelectedBones.Num(); ++SelBoneIdx)
+	{
+		int32 Bone = SelectedBones[SelBoneIdx];
+		if (SimulationType[Bone] == FGeometryCollection::ESimulationTypes::FST_None)
+		{
+			int32 Parent = Parents[Bone];
+			if (Parent != INDEX_NONE)
+			{
+				SelectedBones[SelBoneIdx] = Parent;
+			}
+			else // embedded should always have a parent, but if it somehow does not, just remove from selection
+			{
+				SelectedBones.RemoveAtSwap(SelBoneIdx, 1, false);
+				--SelBoneIdx; // reconsider swapped-in-element at this idx next iter
+			}
+		}
+	}
+}
+
 void FFractureToolContext::ConvertSelectionToRigidNodes()
 {
 	Sanitize();
diff --git a/Engine/Plugins/Experimental/ChaosEditor/Source/FractureEditor/Private/FractureToolContext.h b/Engine/Plugins/Experimental/ChaosEditor/Source/FractureEditor/Private/FractureToolContext.h
--- a/Engine/Plugins/Experimental/ChaosEditor/Source/FractureEditor/Private/FractureToolContext.h
+++ b/Engine/Plugins/Experimental/ChaosEditor/Source/FractureEditor/Private/FractureToolContext.h
@@ -24,6 +24,8 @@ public:
 	void ConvertSelectionToRigidNodes();
 	void ConvertSelectionToEmbeddedGeometryNodes();
 	void ConvertSelectionToClusterNodes();
+	void RemoveChildrenOfSelectedNodesFromSelection();
+	void ConvertEmbeddedSelectionToParents();
 	void RemoveRootNodes();
 	/// @param ProbToKeep The chance to keep each bone in the selection
 	void RandomReduceSelection(float ProbToKeep);
