Fracture tool : break on event data cannot be assigned to root cluster anymore

diff --git a/Engine/Plugins/Experimental/ChaosEditor/Source/FractureEditor/Private/FractureToolProperties.cpp b/Engine/Plugins/Experimental/ChaosEditor/Source/FractureEditor/Private/FractureToolProperties.cpp
--- a/Engine/Plugins/Experimental/ChaosEditor/Source/FractureEditor/Private/FractureToolProperties.cpp
+++ b/Engine/Plugins/Experimental/ChaosEditor/Source/FractureEditor/Private/FractureToolProperties.cpp
@@ -160,7 +160,15 @@ void UFractureToolSetRemoveOnBreak::Execute(TWeakPtr<FFractureEditorModeToolkit>
 
 	 					for (int32 Index : SelectedBones)
 	 					{
-	 						RemoveOnBreak[Index] = RemoveOnBreakData.GetPackedData();
+							// if root bone, then do not set 
+							if (GeometryCollection->Parent[Index] == INDEX_NONE)
+							{
+								RemoveOnBreak[Index] = FRemoveOnBreakData::DisabledPackedData;
+							}
+							else
+							{
+								RemoveOnBreak[Index] = RemoveOnBreakData.GetPackedData();
+							}
 	 					}
 	 				}
 	 			}
