Add more options for controlling radial fracture

diff --git a/Engine/Plugins/Experimental/ChaosEditor/Source/FractureEditor/Private/FractureToolRadial.cpp b/Engine/Plugins/Experimental/ChaosEditor/Source/FractureEditor/Private/FractureToolRadial.cpp
--- a/Engine/Plugins/Experimental/ChaosEditor/Source/FractureEditor/Private/FractureToolRadial.cpp
+++ b/Engine/Plugins/Experimental/ChaosEditor/Source/FractureEditor/Private/FractureToolRadial.cpp
@@ -65,7 +65,6 @@ TArray<UObject*> UFractureToolRadial::GetSettingsObjects() const
 
 void UFractureToolRadial::GenerateVoronoiSites(const FFractureToolContext& Context, TArray<FVector>& Sites)
 {
- 	const FVector::FReal RadialStep = RadialSettings->Radius / RadialSettings->RadialSteps;
 	const FVector::FReal AngularStep = 2 * PI / RadialSettings->AngularSteps;
 
 	FBox Bounds = Context.GetWorldBounds();
@@ -86,14 +85,76 @@ void UFractureToolRadial::GenerateVoronoiSites(const FFractureToolContext& Conte
 		Center = Transform.GetTranslation();
 	}
 
-	FVector::FReal Len = RadialStep * .5;
-	for (int32 ii = 0; ii < RadialSettings->RadialSteps; ++ii, Len += RadialStep)
+	// Precompute consistent noise for each angular step
+	TArray<FVector::FReal> AngleStepOffsets;
+	AngleStepOffsets.SetNumUninitialized(RadialSettings->AngularSteps);
+	for (int32 AngleIdx = 0; AngleIdx < RadialSettings->AngularSteps; ++AngleIdx)
 	{
+		AngleStepOffsets[AngleIdx] = FMath::DegreesToRadians(RandStream.FRandRange(-1, 1) * RadialSettings->AngularNoise);
+	}
+
+	// Compute radial positions following an (idx+1)^exp curve, and then re-normalize back to the Radius range
+	TArray<FVector::FReal> RadialPositions;
+	RadialPositions.SetNumUninitialized(RadialSettings->RadialSteps);
+	FVector::FReal StepOffset = 0;
+	for (int32 RadIdx = 0; RadIdx < RadialSettings->RadialSteps; ++RadIdx)
+	{
+		FVector::FReal RadialPos = FMath::Pow(RadIdx + 1, RadialSettings->RadialStepExponent) + StepOffset;
+		if (RadIdx == 0)
+		{
+			// Note we bring the first point a half-step toward the center, and shift all subsequent points accordingly
+			// so that for Exponent==1, the step from center to first boundary is the same distance as the step between each boundary
+			// (this is only necessary because there is no Voronoi site at the center)
+			RadialPos *= .5;
+			StepOffset = -RadialPos;
+		}
+
+		RadialPositions[RadIdx] = RadialPos;
+	}
+	// Normalize positions so that the diagram fits in the target radius
+	FVector::FReal RadialPosNorm = RadialSettings->Radius / RadialPositions.Last();
+	for (FVector::FReal& RadialPos : RadialPositions)
+	{
+		RadialPos = RadialPos * RadialPosNorm;
+	}
+	// Add radial noise 
+	for (int32 RadIdx = 0; RadIdx < RadialSettings->RadialSteps; ++RadIdx)
+	{
+		FVector::FReal& RadialPos = RadialPositions[RadIdx];
+		// Offset by RadialNoise, but don't allow noise to take the value below 0
+		RadialPos += RandStream.FRandRange(-FMath::Min(RadialPos, RadialSettings->RadialNoise), RadialSettings->RadialNoise);
+	}
+	// make sure the positions remain in increasing order
+	RadialPositions.Sort();
+	// Adjust positions so they are never closer than the RadialMinStep
+	FVector::FReal LastRadialPos = 0;
+	for (int32 RadIdx = 0; RadIdx < RadialSettings->RadialSteps; ++RadIdx)
+	{
+		FVector::FReal MinStep = RadialSettings->RadialMinStep;
+		if (RadIdx == 0)
+		{
+			MinStep *= .5;
+		}
+		if (RadialPositions[RadIdx] - LastRadialPos < MinStep)
+		{
+			RadialPositions[RadIdx] = LastRadialPos + MinStep;
+		}
+		LastRadialPos = RadialPositions[RadIdx];
+	}
+
+	// Create the radial Voronoi sites
+	for (int32 ii = 0; ii < RadialSettings->RadialSteps; ++ii)
+	{
+		FVector::FReal Len = RadialPositions[ii];
 		FVector::FReal Angle = FMath::DegreesToRadians(RadialSettings->AngleOffset);
 		for (int32 kk = 0; kk < RadialSettings->AngularSteps; ++kk, Angle += AngularStep)
 		{
-			FVector RotatingOffset = Len * (FMath::Cos(Angle) * BasisX + FMath::Sin(Angle) * BasisY);
-			Sites.Emplace(Center + RotatingOffset + (RandStream.VRand() * RandStream.FRand() * RadialSettings->Variability));
+			// Add the global noise and the per-point noise into the angle
+			FVector::FReal UseAngle = Angle + AngleStepOffsets[kk] + FMath::DegreesToRadians(RandStream.FRand() * RadialSettings->AngularVariability);
+			// Add per point noise into the radial position
+			FVector::FReal UseRadius = Len + RandStream.FRand() * RadialSettings->RadialVariability;
+			FVector RotatingOffset = UseRadius * (FMath::Cos(UseAngle) * BasisX + FMath::Sin(UseAngle) * BasisY);
+			Sites.Emplace(Center + RotatingOffset + UpVector * (RandStream.FRandRange(-1, 1) * RadialSettings->AxialVariability));
 		}
 	}
 }
diff --git a/Engine/Plugins/Experimental/ChaosEditor/Source/FractureEditor/Private/FractureToolRadial.h b/Engine/Plugins/Experimental/ChaosEditor/Source/FractureEditor/Private/FractureToolRadial.h
--- a/Engine/Plugins/Experimental/ChaosEditor/Source/FractureEditor/Private/FractureToolRadial.h
+++ b/Engine/Plugins/Experimental/ChaosEditor/Source/FractureEditor/Private/FractureToolRadial.h
@@ -17,13 +17,8 @@ public:
 
 	UFractureRadialSettings(const FObjectInitializer& ObjInit)
 		: Super(ObjInit)
-		, Center(FVector(0,0,0))
+		, Center(FVector(0, 0, 0))
 		, Normal(FVector(0, 0, 1))
-		, Radius(50.0f)
-		, AngularSteps(5)
-		, RadialSteps(5)
-		, AngleOffset(0.0f)
-		, Variability(0.0f)
 	{}
 
 	/** Center of generated pattern. Only used when "Use Gizmo" is disabled */
@@ -37,25 +32,50 @@ public:
 	UPROPERTY()
 	bool bPositionedByGizmo = true;
 
-	/** Pattern radius (in cm) */
-	UPROPERTY(EditAnywhere, Category = RadialVoronoi, meta = (DisplayName = "Radius", UIMin = "0.0", ClampMin = "0.0"))
-	float Radius;
-
 	/** Number of angular steps */
 	UPROPERTY(EditAnywhere, Category = RadialVoronoi, meta = (DisplayName = "Angular Steps", UIMin = "1", UIMax = "50", ClampMin = "1"))
-	int AngularSteps;
+	int AngularSteps = 5;
+
+	/** Angle offset at each radial step (in degrees) */
+	UPROPERTY(EditAnywhere, Category = RadialVoronoi, meta = (DisplayName = "Angle Offset"))
+	float AngleOffset = 0.0f;
+
+	/** Amount of global variation to apply to each angular step (in degrees) */
+	UPROPERTY(EditAnywhere, Category = RadialVoronoi)
+	float AngularNoise = 0.f;
+
+	/** Pattern radius (in cm) */
+	UPROPERTY(EditAnywhere, Category = RadialVoronoi, meta = (DisplayName = "Radius", UIMin = "0.0", ClampMin = "0.0"))
+	float Radius = 50.0f;
 
 	/** Number of radial steps */
 	UPROPERTY(EditAnywhere, Category = RadialVoronoi, meta = (DisplayName = "Radial Steps", UIMin = "1", UIMax = "50", ClampMin = "1"))
-	int RadialSteps;
+	int RadialSteps = 5;
 
-	/** Angle offset at each radial step (in degrees) */
-	UPROPERTY(EditAnywhere, Category = RadialVoronoi, meta = (DisplayName = "Angle Offset"))
-	float AngleOffset;
+	/** Radial steps will follow a distribution based on this exponent, i.e., Pow(distance from center, RadialStepExponent) */
+	UPROPERTY(EditAnywhere, Category = RadialVoronoi, meta = (UIMin = ".01", UIMax = "10", ClampMin = ".01", ClampMax = "20"))
+	float RadialStepExponent = 1.f;
+
+	/** Minimum radial separation between any two voronoi points (in cm) */
+	UPROPERTY(EditAnywhere, Category = RadialVoronoi, meta = (UIMin = ".25", ClampMin = ".01"))
+	float RadialMinStep = 1.f;
+
+	/** Amount of global variation to apply to each radial step (in cm) */
+	UPROPERTY(EditAnywhere, Category = RadialVoronoi, meta = (ClampMin = "0"))
+	float RadialNoise = 0.f;
+
+
+	/** Amount to randomly displace each Voronoi site radially (in cm) */
+	UPROPERTY(EditAnywhere, Category = PerPointVariability, meta = (ClampMin = "0.0"))
+	float RadialVariability = 0.f;
+
+	/** Amount to randomly displace each Voronoi site in angle (in degrees) */
+	UPROPERTY(EditAnywhere, Category = PerPointVariability, meta = (ClampMin = "0.0"))
+	float AngularVariability = 0.f;
 
-	/** Amount to randomly displace each Voronoi site (in cm) */
-	UPROPERTY(EditAnywhere, Category = RadialVoronoi, meta = (DisplayName = "Variability", UIMin = "0.0", ClampMin = "0.0"))
-	float Variability;
+	/** Amount to randomly displace each Voronoi site in the direction of the rotation axis (in cm) */
+	UPROPERTY(EditAnywhere, Category = PerPointVariability, meta = (ClampMin = "0.0"))
+	float AxialVariability = 0.f;
 };
 
 
