make fracture "merge" cluster tool also work on leaf nodes, so users can easily merge a node into a cluster

diff --git a/Engine/Plugins/Experimental/ChaosEditor/Source/FractureEditor/Private/FractureToolClustering.cpp b/Engine/Plugins/Experimental/ChaosEditor/Source/FractureEditor/Private/FractureToolClustering.cpp
--- a/Engine/Plugins/Experimental/ChaosEditor/Source/FractureEditor/Private/FractureToolClustering.cpp
+++ b/Engine/Plugins/Experimental/ChaosEditor/Source/FractureEditor/Private/FractureToolClustering.cpp
@@ -133,13 +133,7 @@ void UFractureToolCluster::Execute(TWeakPtr<FFractureEditorModeToolkit> InToolki
 
 				if (LowestCommonAncestor != INDEX_NONE)
 				{
-					// ClusterBonesUnderNewNode expects a sibling of the new cluster so we require a child node of the common ancestor.
-					const TManagedArray<TSet<int32>>& Children = Context.GetGeometryCollection()->GetAttribute<TSet<int32>>("Children", FGeometryCollection::TransformGroup);
-					TArray<int32> Siblings = Children[LowestCommonAncestor].Array();
-					if (Siblings.Num() > 0)
-					{
-						FGeometryCollectionClusteringUtility::ClusterBonesUnderNewNode(Context.GetGeometryCollection().Get(), Siblings[0], Context.GetSelection(), true);
-					}
+					FGeometryCollectionClusteringUtility::ClusterBonesUnderNewNodeWithParent(Context.GetGeometryCollection().Get(), LowestCommonAncestor, Context.GetSelection(), true);
 				}
 
 				Context.GenerateGuids(StartTransformCount);
@@ -271,7 +265,7 @@ FText UFractureToolClusterMerge::GetDisplayText() const
 
 FText UFractureToolClusterMerge::GetTooltipText() const
 {
-	return FText(NSLOCTEXT("FractureToolClusteringOps", "FractureToolClusterMergeTooltip", "Merge selected clusters."));
+	return FText(NSLOCTEXT("FractureToolClusteringOps", "FractureToolClusterMergeTooltip", "Merge selected nodes into the same cluster."));
 }
 
 FSlateIcon UFractureToolClusterMerge::GetToolIcon() const
@@ -281,7 +275,7 @@ FSlateIcon UFractureToolClusterMerge::GetToolIcon() const
 
 void UFractureToolClusterMerge::RegisterUICommand(FFractureEditorCommands* BindingContext)
 {
-	UI_COMMAND_EXT(BindingContext, UICommandInfo, "Merge", "Merge", "Merge selected clusters.", EUserInterfaceActionType::Button, FInputChord());
+	UI_COMMAND_EXT(BindingContext, UICommandInfo, "Merge", "Merge", "Merge selected nodes into the same cluster.", EUserInterfaceActionType::Button, FInputChord());
 	BindingContext->ClusterMerge = UICommandInfo;
 }
 
@@ -301,24 +295,50 @@ void UFractureToolClusterMerge::Execute(TWeakPtr<FFractureEditorModeToolkit> InT
 
 			const TManagedArray<TSet<int32>>& Children = Context.GetGeometryCollection()->Children;
 
-			Context.ConvertSelectionToClusterNodes();
+			Context.ConvertEmbeddedSelectionToParents(); // embedded geo must stay attached to parent
 
 			// Collect children of context clusters
 			TArray<int32> ChildBones;
 			ChildBones.Reserve(Context.GetGeometryCollection()->NumElements(FGeometryCollection::TransformGroup));
+			bool bHasClusters = false;
 			for (int32 Select : Context.GetSelection())
 			{
-				ChildBones.Append(Children[Select].Array());
+				if (Context.GetGeometryCollection()->SimulationType[Select] == FGeometryCollection::ESimulationTypes::FST_Clustered)
+				{
+					bHasClusters = true;
+					ChildBones.Append(Children[Select].Array());
+				}
+				else
+				{
+					ChildBones.Add(Select);
+				}
 			}
 			
-			int32 MergeNode = FGeometryCollectionClusteringUtility::PickBestNodeToMergeTo(Context.GetGeometryCollection().Get(), Context.GetSelection());
-			if (MergeNode >= 0)
+			// If there were no clusters in the selection, we create a cluster -- i.e., follow the logic of UFractureToolCluster::Execute
+			if (!bHasClusters)
 			{
-				FGeometryCollectionClusteringUtility::ClusterBonesUnderExistingNode(Context.GetGeometryCollection().Get(), MergeNode, ChildBones);
-				FGeometryCollectionClusteringUtility::RemoveDanglingClusters(Context.GetGeometryCollection().Get());
+				// Cluster selected bones beneath common parent
+				int32 LowestCommonAncestor = FGeometryCollectionClusteringUtility::FindLowestCommonAncestor(Context.GetGeometryCollection().Get(), Context.GetSelection());
+				if (LowestCommonAncestor != INDEX_NONE)
+				{
+					int32 StartTransformCount = Context.GetGeometryCollection()->Transform.Num();
+					FGeometryCollectionClusteringUtility::ClusterBonesUnderNewNodeWithParent(Context.GetGeometryCollection().Get(), LowestCommonAncestor, Context.GetSelection(), true);
+					Context.GenerateGuids(StartTransformCount);
+					Context.SetSelection({ LowestCommonAncestor });
+					Refresh(Context, Toolkit);
+				}
+			}
+			else
+			{
+				int32 MergeNode = FGeometryCollectionClusteringUtility::PickBestNodeToMergeTo(Context.GetGeometryCollection().Get(), Context.GetSelection());
+				if (MergeNode >= 0)
+				{
+					FGeometryCollectionClusteringUtility::ClusterBonesUnderExistingNode(Context.GetGeometryCollection().Get(), MergeNode, ChildBones);
+					FGeometryCollectionClusteringUtility::RemoveDanglingClusters(Context.GetGeometryCollection().Get());
 
-				Context.SetSelection({ MergeNode });
-				Refresh(Context, Toolkit);
+					Context.SetSelection({ MergeNode });
+					Refresh(Context, Toolkit);
+				}
 			}
 		}
 
diff --git a/Engine/Source/Runtime/Experimental/Chaos/Private/GeometryCollection/GeometryCollectionClusteringUtility.cpp b/Engine/Source/Runtime/Experimental/Chaos/Private/GeometryCollection/GeometryCollectionClusteringUtility.cpp
--- a/Engine/Source/Runtime/Experimental/Chaos/Private/GeometryCollection/GeometryCollectionClusteringUtility.cpp
+++ b/Engine/Source/Runtime/Experimental/Chaos/Private/GeometryCollection/GeometryCollectionClusteringUtility.cpp
@@ -13,6 +13,15 @@ void FGeometryCollectionClusteringUtility::ClusterBonesUnderNewNode(FGeometryCol
 {
 	check(GeometryCollection);
 
+	TManagedArray<int32>& Parents = GeometryCollection->Parent;
+	int32 ParentIdx = Parents[InsertAtIndex];
+	return ClusterBonesUnderNewNodeWithParent(GeometryCollection, ParentIdx, SelectedBones, CalcNewLocalTransform, Validate);
+}
+
+void FGeometryCollectionClusteringUtility::ClusterBonesUnderNewNodeWithParent(FGeometryCollection* GeometryCollection, const int32 ParentOfNewNode, const TArray<int32>& SelectedBones, bool CalcNewLocalTransform, bool Validate)
+{
+	check(GeometryCollection);
+
 
 	TManagedArray<FTransform>& Transforms = GeometryCollection->Transform;
 	TManagedArray<FString>& BoneNames = GeometryCollection->BoneName;
@@ -23,9 +32,7 @@ void FGeometryCollectionClusteringUtility::ClusterBonesUnderNewNode(FGeometryCol
 	// insert a new node between the selected bones and their shared parent
 	int NewBoneIndex = GeometryCollection->AddElements(1, FGeometryCollection::TransformGroup);
 
-	// New Bone Setup takes level/parent from the first of the Selected Bones
-	int32 SourceBoneIndex = InsertAtIndex;
-	int32 OriginalParentIndex = Parents[SourceBoneIndex];
+	int32 OriginalParentIndex = ParentOfNewNode;
 	Parents[NewBoneIndex] = OriginalParentIndex;
 	Children[NewBoneIndex] = TSet<int32>(SelectedBones);
 	SimType[NewBoneIndex] = FGeometryCollection::ESimulationTypes::FST_Clustered;
diff --git a/Engine/Source/Runtime/Experimental/Chaos/Public/GeometryCollection/GeometryCollectionClusteringUtility.h b/Engine/Source/Runtime/Experimental/Chaos/Public/GeometryCollection/GeometryCollectionClusteringUtility.h
--- a/Engine/Source/Runtime/Experimental/Chaos/Public/GeometryCollection/GeometryCollectionClusteringUtility.h
+++ b/Engine/Source/Runtime/Experimental/Chaos/Public/GeometryCollection/GeometryCollectionClusteringUtility.h
@@ -35,7 +35,10 @@ public:
 	*
 	* Node E has no geometry of its own, only a transform by which to control A & B as a single unit
 	*/
-	static void ClusterBonesUnderNewNode(FGeometryCollection* GeometryCollection, int32 InsertAtIndex, const TArray<int32>& SelectedBones, bool CalcNewLocalTransform, bool Validate = true);
+	static void ClusterBonesUnderNewNode(FGeometryCollection* GeometryCollection, const int32 InsertAtIndex, const TArray<int32>& SelectedBones, bool CalcNewLocalTransform, bool Validate = true);
+
+	// Same as ClusterBonesUnderNewNode, but specify the parent of the new node instead of a sibling
+	static void ClusterBonesUnderNewNodeWithParent(FGeometryCollection* GeometryCollection, const int32 ParentOfNewNode, const TArray<int32>& SelectedBones, bool CalcNewLocalTransform, bool Validate = true);
 
 	/** Cluster all existing bones under a new root node, so there is now only one root node and a completely flat hierarchy underneath it */
 	static void ClusterAllBonesUnderNewRoot(FGeometryCollection* GeometryCollection);
