Add a tiny amount of radial variation to the fracture radial tool (if radius itself is not tiny), to work around an apparent failure case in Voro++ for some point patterns.

diff --git a/Engine/Plugins/Experimental/ChaosEditor/Source/FractureEditor/Private/FractureToolRadial.cpp b/Engine/Plugins/Experimental/ChaosEditor/Source/FractureEditor/Private/FractureToolRadial.cpp
--- a/Engine/Plugins/Experimental/ChaosEditor/Source/FractureEditor/Private/FractureToolRadial.cpp
+++ b/Engine/Plugins/Experimental/ChaosEditor/Source/FractureEditor/Private/FractureToolRadial.cpp
@@ -142,6 +142,11 @@ void UFractureToolRadial::GenerateVoronoiSites(const FFractureToolContext& Conte
 		LastRadialPos = RadialPositions[RadIdx];
 	}
 
+	// Add a bit of noise to work around failure case in Voro++
+	// TODO: fix the failure case in Voro++ and remove this
+	float MinRadialVariability = RadialSettings->Radius > 1.f ? .0001f : 0.f;
+	float UseRadialVariability = FMath::Max(MinRadialVariability, RadialSettings->RadialVariability);
+
 	// Create the radial Voronoi sites
 	for (int32 ii = 0; ii < RadialSettings->RadialSteps; ++ii)
 	{
@@ -152,7 +157,7 @@ void UFractureToolRadial::GenerateVoronoiSites(const FFractureToolContext& Conte
 			// Add the global noise and the per-point noise into the angle
 			FVector::FReal UseAngle = Angle + AngleStepOffsets[kk] + FMath::DegreesToRadians(RandStream.FRand() * RadialSettings->AngularVariability);
 			// Add per point noise into the radial position
-			FVector::FReal UseRadius = Len + RandStream.FRand() * RadialSettings->RadialVariability;
+			FVector::FReal UseRadius = Len + FVector::FReal(RandStream.FRand() * UseRadialVariability);
 			FVector RotatingOffset = UseRadius * (FMath::Cos(UseAngle) * BasisX + FMath::Sin(UseAngle) * BasisY);
 			Sites.Emplace(Center + RotatingOffset + UpVector * (RandStream.FRandRange(-1, 1) * RadialSettings->AxialVariability));
 		}
