Correct the wording used to explain the Geometry Collection "Split Components" feature

diff --git a/Engine/Plugins/Experimental/ChaosEditor/Source/FractureEditor/Private/FractureToolGenerators.cpp b/Engine/Plugins/Experimental/ChaosEditor/Source/FractureEditor/Private/FractureToolGenerators.cpp
--- a/Engine/Plugins/Experimental/ChaosEditor/Source/FractureEditor/Private/FractureToolGenerators.cpp
+++ b/Engine/Plugins/Experimental/ChaosEditor/Source/FractureEditor/Private/FractureToolGenerators.cpp
@@ -305,7 +305,7 @@ void SCreateGeometryCollectionFromObject::Construct(const FArguments& InArgs, TS
 					[
 						SNew(STextBlock)
 						.Text(LOCTEXT("SplitComponentsLabel", "Split Components"))
-						.ToolTipText(LOCTEXT("SplitComponentsToolTip", "If checked, connected components will be split into separate bones on import."))
+						.ToolTipText(LOCTEXT("SplitComponentsToolTip", "If checked, triangles that are not topologically connected will be assigned separate bones on import."))
 					]
 
 				+ SHorizontalBox::Slot()
@@ -324,7 +324,7 @@ void SCreateGeometryCollectionFromObject::Construct(const FArguments& InArgs, TS
 							}
 							bSplitIslands = NewState == ECheckBoxState::Checked;
 						})
-						.ToolTipText(LOCTEXT("SplitComponentsToolTip", "If checked, connected components will be split into separate bones on import."))
+						.ToolTipText(LOCTEXT("SplitComponentsToolTip", "If checked, triangles that are not topologically connected will be assigned separate bones on import."))
 					]
 				]
 			]
diff --git a/Engine/Source/Runtime/Experimental/GeometryCollectionEngine/Public/GeometryCollection/GeometryCollectionObject.h b/Engine/Source/Runtime/Experimental/GeometryCollectionEngine/Public/GeometryCollection/GeometryCollectionObject.h
--- a/Engine/Source/Runtime/Experimental/GeometryCollectionEngine/Public/GeometryCollection/GeometryCollectionObject.h
+++ b/Engine/Source/Runtime/Experimental/GeometryCollectionEngine/Public/GeometryCollection/GeometryCollectionObject.h
@@ -38,7 +38,7 @@ struct GEOMETRYCOLLECTIONENGINE_API FGeometryCollectionSource
 	UPROPERTY(BlueprintReadWrite, EditAnywhere, Category = "GeometrySource")
 	bool bAddInternalMaterials = true;
 
-	/** Whether connected components in the source object should be added as separate pieces of geometry. */
+	/** Whether individual source mesh components should be split into separate pieces of geometry based on mesh connectivity. If checked, triangles that are not topologically connected will be assigned separate bones. */
 	UPROPERTY(BlueprintReadWrite, EditAnywhere, Category = "GeometrySource")
 	bool bSplitComponents = false;
 };
