Add new geometry collection proximity detection settings to the fracture mode project settings

diff --git a/Engine/Plugins/Experimental/ChaosEditor/Source/FractureEditor/Private/FractureModeSettings.cpp b/Engine/Plugins/Experimental/ChaosEditor/Source/FractureEditor/Private/FractureModeSettings.cpp
--- a/Engine/Plugins/Experimental/ChaosEditor/Source/FractureEditor/Private/FractureModeSettings.cpp
+++ b/Engine/Plugins/Experimental/ChaosEditor/Source/FractureEditor/Private/FractureModeSettings.cpp
@@ -4,6 +4,7 @@
 
 #include "GeometryCollection/GeometryCollection.h"
 #include "GeometryCollection/GeometryCollectionConvexPropertiesInterface.h"
+#include "GeometryCollection/GeometryCollectionProximityPropertiesInterface.h"
 
 #include UE_INLINE_GENERATED_CPP_BY_NAME(FractureModeSettings)
 
@@ -17,3 +18,18 @@ void UFractureModeSettings::ApplyDefaultConvexSettings(FGeometryCollection& Geom
 	Properties.OverlapRemovalShrinkPercent = ConvexOverlapRemovalShrinkPercent;
 	GeometryCollection.SetConvexProperties(Properties);
 }
+
+void UFractureModeSettings::ApplyDefaultProximitySettings(FGeometryCollection& GeometryCollection) const
+{
+	FGeometryCollectionProximityPropertiesInterface::FProximityProperties Properties = GeometryCollection.GetProximityProperties();
+	Properties.bUseAsConnectionGraph = bProximityUseAsConnectionGraph;
+	Properties.DistanceThreshold = ProximityDistanceThreshold;
+	Properties.Method = ProximityMethod;
+	GeometryCollection.SetProximityProperties(Properties);
+}
+
+void UFractureModeSettings::ApplyDefaultSettings(FGeometryCollection& GeometryCollection) const
+{
+	ApplyDefaultProximitySettings(GeometryCollection);
+	ApplyDefaultConvexSettings(GeometryCollection);
+}
diff --git a/Engine/Plugins/Experimental/ChaosEditor/Source/FractureEditor/Private/FractureToolGenerators.cpp b/Engine/Plugins/Experimental/ChaosEditor/Source/FractureEditor/Private/FractureToolGenerators.cpp
--- a/Engine/Plugins/Experimental/ChaosEditor/Source/FractureEditor/Private/FractureToolGenerators.cpp
+++ b/Engine/Plugins/Experimental/ChaosEditor/Source/FractureEditor/Private/FractureToolGenerators.cpp
@@ -686,7 +686,7 @@ AGeometryCollectionActor* UFractureToolGenerateAsset::ConvertActorsToGeometryCol
 	::GeometryCollection::GenerateTemporaryGuids(FracturedGeometryCollection->GetGeometryCollection().Get(), 0 , true);
 
 	const UFractureModeSettings* ModeSettings = GetDefault<UFractureModeSettings>();
-	ModeSettings->ApplyDefaultConvexSettings(*FracturedGeometryCollection->GetGeometryCollection());
+	ModeSettings->ApplyDefaultSettings(*FracturedGeometryCollection->GetGeometryCollection());
 
 	return NewActor;
 }
@@ -874,7 +874,7 @@ void UFractureToolResetAsset::Execute(TWeakPtr<FFractureEditorModeToolkit> InToo
 				}
 
 				const UFractureModeSettings* ModeSettings = GetDefault<UFractureModeSettings>();
-				ModeSettings->ApplyDefaultConvexSettings(*GeometryCollection);
+				ModeSettings->ApplyDefaultSettings(*GeometryCollection);
 
 				FGeometryCollectionClusteringUtility::UpdateHierarchyLevelOfChildren(GeometryCollection, -1);
 				AddSingleRootNodeIfRequired(GeometryCollectionObject);
diff --git a/Engine/Plugins/Experimental/ChaosEditor/Source/FractureEditor/Private/FractureToolProximity.cpp b/Engine/Plugins/Experimental/ChaosEditor/Source/FractureEditor/Private/FractureToolProximity.cpp
--- a/Engine/Plugins/Experimental/ChaosEditor/Source/FractureEditor/Private/FractureToolProximity.cpp
+++ b/Engine/Plugins/Experimental/ChaosEditor/Source/FractureEditor/Private/FractureToolProximity.cpp
@@ -14,12 +14,36 @@
 #define LOCTEXT_NAMESPACE "FractureToolProximity"
 
 
+void UFractureProximityActions::SaveAsDefaults()
+{
+	UFractureToolProximity* ProximityTool = Cast<UFractureToolProximity>(OwnerTool.Get());
+	UFractureModeSettings* ModeSettings = GetMutableDefault<UFractureModeSettings>();
+	ModeSettings->Modify();
+	ModeSettings->ProximityMethod = ProximityTool->ProximitySettings->Method;
+	ModeSettings->ProximityDistanceThreshold = ProximityTool->ProximitySettings->DistanceThreshold;
+	ModeSettings->bProximityUseAsConnectionGraph = ProximityTool->ProximitySettings->bUseAsConnectionGraph;
+}
+
+void UFractureProximityActions::SetFromDefaults()
+{
+	UFractureToolProximity* ProximityTool = Cast<UFractureToolProximity>(OwnerTool.Get());
+	const UFractureModeSettings* ModeSettings = GetDefault<UFractureModeSettings>();
+	ProximityTool->ProximitySettings->Method = ModeSettings->ProximityMethod;
+	ProximityTool->ProximitySettings->DistanceThreshold = ModeSettings->ProximityDistanceThreshold;
+	ProximityTool->ProximitySettings->bUseAsConnectionGraph = ModeSettings->bProximityUseAsConnectionGraph;
+
+	ProximityTool->NotifyOfPropertyChangeByTool(this);
+}
+
+
 
 UFractureToolProximity::UFractureToolProximity(const FObjectInitializer& ObjInit)
 	: Super(ObjInit)
 {
 	ProximitySettings = NewObject<UFractureProximitySettings>(GetTransientPackage(), UFractureProximitySettings::StaticClass());
 	ProximitySettings->OwnerTool = this;
+	ProximityActions = NewObject<UFractureProximityActions>(GetTransientPackage(), UFractureProximityActions::StaticClass());
+	ProximityActions->OwnerTool = this;
 }
 
 bool UFractureToolProximity::CanExecute() const
@@ -57,6 +81,7 @@ TArray<UObject*> UFractureToolProximity::GetSettingsObjects() const
 {
 	TArray<UObject*> Settings;
 	Settings.Add(ProximitySettings);
+	Settings.Add(ProximityActions);
 	return Settings;
 }
 
diff --git a/Engine/Plugins/Experimental/ChaosEditor/Source/FractureEditor/Private/FractureToolProximity.h b/Engine/Plugins/Experimental/ChaosEditor/Source/FractureEditor/Private/FractureToolProximity.h
--- a/Engine/Plugins/Experimental/ChaosEditor/Source/FractureEditor/Private/FractureToolProximity.h
+++ b/Engine/Plugins/Experimental/ChaosEditor/Source/FractureEditor/Private/FractureToolProximity.h
@@ -44,6 +44,31 @@ public:
 };
 
 
+/**
+ * UFUNCTION actions to manage default proximity detection settings for geometry collections
+ * (These are pulled out from the above settings object mainly to control their ordering in the properties panel)
+ */
+UCLASS(config = EditorPerProjectUserSettings)
+class UFractureProximityActions : public UFractureToolSettings
+{
+public:
+
+	GENERATED_BODY()
+
+	UFractureProximityActions(const FObjectInitializer& ObjInit)
+		: Super(ObjInit)
+	{}
+
+	/** Save settings as project defaults, to be used for all new geometry collections */
+	UFUNCTION(CallInEditor, Category = ProjectDefaults, meta = (DisplayName = "Save As Defaults"))
+	void SaveAsDefaults();
+
+	/** Set settings from current project defaults */
+	UFUNCTION(CallInEditor, Category = ProjectDefaults, meta = (DisplayName = "Set From Defaults"))
+	void SetFromDefaults();
+
+};
+
 
 UCLASS(DisplayName = "Proximity Tool", Category = "FractureTools")
 class UFractureToolProximity : public UFractureModalTool
@@ -79,10 +104,13 @@ public:
 
 	virtual void Setup() override;
 
+	UPROPERTY(EditAnywhere, Category = Proximity)
+	TObjectPtr<UFractureProximitySettings> ProximitySettings;
+
 private:
 
 	UPROPERTY(EditAnywhere, Category = Proximity)
-	TObjectPtr<UFractureProximitySettings> ProximitySettings;
+	TObjectPtr<UFractureProximityActions> ProximityActions;
 
 	void UpdateVisualizations();
 	virtual void ClearVisualizations() override
diff --git a/Engine/Plugins/Experimental/ChaosEditor/Source/FractureEditor/Public/FractureModeSettings.h b/Engine/Plugins/Experimental/ChaosEditor/Source/FractureEditor/Public/FractureModeSettings.h
--- a/Engine/Plugins/Experimental/ChaosEditor/Source/FractureEditor/Public/FractureModeSettings.h
+++ b/Engine/Plugins/Experimental/ChaosEditor/Source/FractureEditor/Public/FractureModeSettings.h
@@ -6,6 +6,7 @@
 #include "Engine/EngineTypes.h"
 
 #include "GeometryCollection/GeometryCollectionConvexUtility.h"
+#include "GeometryCollection/GeometryCollectionProximityUtility.h"
 
 #include "FractureModeSettings.generated.h"
 
@@ -52,9 +53,29 @@ public:
 	UPROPERTY(config, EditAnywhere, Category = "Fracture Mode|Convex Generation Defaults|Overlap Removal", meta = (DisplayName = "Max Removal Fraction", ClampMin = ".01", ClampMax = "1"))
 	double ConvexFractionAllowRemove = .5;
 
+
+	/** Default method used to detect proximity of geometry in a new geometry collection */
+	UPROPERTY(config, EditAnywhere, Category = "Fracture Mode|Proximity Detection Defaults", meta = (DisplayName = "Detection Method"))
+	EProximityMethod ProximityMethod = EProximityMethod::Precise;
+
+	/** When Proximity Detection Method is Convex Hull, how close two hulls need to be to be considered as 'in proximity' */
+	UPROPERTY(config, EditAnywhere, Category = "Fracture Mode|Proximity Detection Defaults", meta = (ClampMin = "0", DisplayName = "Distance Threshold", EditCondition = "ProximityMethod == EProximityMethod::ConvexHull"))
+	float ProximityDistanceThreshold = 1;
+
+	/** Whether to automatically transform the proximity graph into a connection graph to be used for simulation */
+	UPROPERTY(config, EditAnywhere, Category = "Fracture Mode|Proximity Detection Defaults", meta = (DisplayName = "Use As Connection Graph"))
+	bool bProximityUseAsConnectionGraph = false;
+	
+
 	// Apply Convex Generation Defaults to a GeometryCollection
 	void ApplyDefaultConvexSettings(FGeometryCollection& GeometryCollection) const;
 
+	// Apply Proximity Detection Defaults to a GeometryCollection
+	void ApplyDefaultProximitySettings(FGeometryCollection& GeometryCollection) const;
+
+	// Apply all default settings to a GeometryCollection
+	void ApplyDefaultSettings(FGeometryCollection& GeometryCollection) const;
+
 
 	DECLARE_MULTICAST_DELEGATE_TwoParams(UFractureModeSettingsModified, UObject*, FProperty*);
 	UFractureModeSettingsModified OnModified;
