Add a background-compute preview mesh to the fracture plane cut tool to show the shape of the noise
 + to support updating the background mesh, make the fracture mode set a realtime override and add an OnTick function to fracture tools
 + expose noise function used by planarcut code to the header file, so it can be used by the preview mesh as well

diff --git a/Engine/Plugins/Experimental/ChaosEditor/Source/FractureEditor/FractureEditor.Build.cs b/Engine/Plugins/Experimental/ChaosEditor/Source/FractureEditor/FractureEditor.Build.cs
--- a/Engine/Plugins/Experimental/ChaosEditor/Source/FractureEditor/FractureEditor.Build.cs
+++ b/Engine/Plugins/Experimental/ChaosEditor/Source/FractureEditor/FractureEditor.Build.cs
@@ -37,6 +37,7 @@ public class FractureEditor : ModuleRules
 				"ModelingComponents",
 				"ModelingOperators",
 				"GeometryCore",
+				"DynamicMesh",
 				"MeshDescription",
 				"StaticMeshDescription",
 				"PlanarCut",
diff --git a/Engine/Plugins/Experimental/ChaosEditor/Source/FractureEditor/Private/FractureEditorMode.cpp b/Engine/Plugins/Experimental/ChaosEditor/Source/FractureEditor/Private/FractureEditorMode.cpp
--- a/Engine/Plugins/Experimental/ChaosEditor/Source/FractureEditor/Private/FractureEditorMode.cpp
+++ b/Engine/Plugins/Experimental/ChaosEditor/Source/FractureEditor/Private/FractureEditorMode.cpp
@@ -31,7 +31,7 @@
 #include "BaseGizmos/TransformGizmoUtil.h"
 #include "Snapping/ModelingSceneSnappingManager.h"
 #include "Components/InstancedStaticMeshComponent.h"
-
+#include "SLevelViewport.h"
 
 #include "UnrealEdGlobals.h"
 #include "EditorModeManager.h"
@@ -87,6 +87,8 @@ void UFractureEditorMode::Enter()
 	
 	OnActorSelectionChanged(SelectedObjects, false);
 
+	ConfigureRealTimeViewportsOverride(true);
+
 }
 
 void UFractureEditorMode::Exit()
@@ -115,6 +117,8 @@ void UFractureEditorMode::Exit()
 		LevelEditor->OnMapChanged().RemoveAll( this );
 	}
 
+	ConfigureRealTimeViewportsOverride(false);
+
 	// Call base Exit method to ensure proper cleanup
 	Super::Exit();
 }
@@ -141,6 +145,52 @@ void UFractureEditorMode::PostRedo(bool bSuccess)
 	OnUndoRedo();
 }
 
+void UFractureEditorMode::ConfigureRealTimeViewportsOverride(bool bEnable)
+{
+	FLevelEditorModule& LevelEditorModule = FModuleManager::GetModuleChecked<FLevelEditorModule>("LevelEditor");
+	TSharedPtr<ILevelEditor> LevelEditor = LevelEditorModule.GetFirstLevelEditor();
+	if (LevelEditor.IsValid())
+	{
+		TArray<TSharedPtr<SLevelViewport>> Viewports = LevelEditor->GetViewports();
+		for (const TSharedPtr<SLevelViewport>& ViewportWindow : Viewports)
+		{
+			if (ViewportWindow.IsValid())
+			{
+				FEditorViewportClient& Viewport = ViewportWindow->GetAssetViewportClient();
+				const FText SystemDisplayName = LOCTEXT("RealtimeOverrideMessage_FractureMode", "Fracture Mode");
+				if (bEnable)
+				{
+					Viewport.AddRealtimeOverride(bEnable, SystemDisplayName);
+				}
+				else
+				{
+					Viewport.RemoveRealtimeOverride(SystemDisplayName, false);
+				}
+			}
+		}
+	}
+}
+
+
+void UFractureEditorMode::Tick(FEditorViewportClient* ViewportClient, float DeltaTime)
+{
+	Super::Tick(ViewportClient, DeltaTime);
+
+	if (!Toolkit.IsValid())
+	{
+		return;
+	}
+
+	// TODO: show a warning if realtime tick is not enabled?
+
+	FFractureEditorModeToolkit* FractureToolkit = (FFractureEditorModeToolkit*)Toolkit.Get();
+
+	if (UFractureModalTool* FractureTool = FractureToolkit->GetActiveTool())
+	{
+		FractureTool->OnTick(DeltaTime);
+	}
+}
+
 void UFractureEditorMode::Render(const FSceneView* View, FViewport* Viewport, FPrimitiveDrawInterface* PDI)
 {
 	Super::Render(View, Viewport, PDI);
@@ -318,7 +368,7 @@ bool UFractureEditorMode::UpdateSelectionInFrustum(const FConvexVolume& InFrustu
 	FTransform ActorTransform = Actor->GetTransform();
 	FMatrix InvActorMatrix(ActorTransform.ToInverseMatrixWithScale());
 
-	FConvexVolume SelectionFrustum(TranformFrustum(InFrustum, InvActorMatrix));
+	FConvexVolume SelectionFrustum(TransformFrustum(InFrustum, InvActorMatrix));
 
 	TArray<FBox> BoundsPerBone;
 	TArray<int32> SelectedBonesArray;
@@ -629,7 +679,7 @@ void UFractureEditorMode::HandlePackageReloaded(const EPackageReloadPhase InPack
 	}
 }
 
-FConvexVolume UFractureEditorMode::TranformFrustum(const FConvexVolume& InFrustum, const FMatrix& InMatrix)
+FConvexVolume UFractureEditorMode::TransformFrustum(const FConvexVolume& InFrustum, const FMatrix& InMatrix)
 {
 	FConvexVolume NewFrustum;
 	NewFrustum.Planes.Empty(6);
diff --git a/Engine/Plugins/Experimental/ChaosEditor/Source/FractureEditor/Private/FractureToolPlaneCut.cpp b/Engine/Plugins/Experimental/ChaosEditor/Source/FractureEditor/Private/FractureToolPlaneCut.cpp
--- a/Engine/Plugins/Experimental/ChaosEditor/Source/FractureEditor/Private/FractureToolPlaneCut.cpp
+++ b/Engine/Plugins/Experimental/ChaosEditor/Source/FractureEditor/Private/FractureToolPlaneCut.cpp
@@ -2,6 +2,7 @@
 
 #include "FractureToolPlaneCut.h"
 #include "FractureEditorStyle.h"
+#include "FractureSettings.h"
 #include "PlanarCut.h"
 #include "GeometryCollection/GeometryCollection.h"
 #include "GeometryCollection/GeometryCollectionObject.h"
@@ -10,6 +11,13 @@
 #include "Drawing/MeshDebugDrawing.h"
 #include "FrameTypes.h"
 #include "FractureToolBackgroundTask.h"
+#include "MeshOpPreviewHelpers.h"
+#include "Materials/MaterialInterface.h"
+#include "ToolSetupUtil.h"
+#include "Generators/RectangleMeshGenerator.h"
+#include "DynamicMesh/MeshNormals.h"
+#include "DynamicMesh/MeshTransforms.h"
+#include "DynamicMeshEditor.h"
 
 #include UE_INLINE_GENERATED_CPP_BY_NAME(FractureToolPlaneCut)
 
@@ -17,6 +25,109 @@ using namespace UE::Fracture;
 
 #define LOCTEXT_NAMESPACE "FracturePlanar"
 
+// This op runs in a background thread to generate a preview of the cutting mesh, with noise applied
+class FPlaneNoisePreviewOp : public UE::Geometry::FDynamicMeshOperator
+{
+public:
+	virtual ~FPlaneNoisePreviewOp() {}
+
+	// inputs
+	TArray<FTransform> PlaneTransforms;
+	TArray<FVector> NoisePivots, ExplodedVectors;
+	TArray<FNoiseOffsets> NoiseOffsets;
+	FNoiseSettings NoiseSettings;
+	float PlaneSize;
+
+	// FDynamicMeshOperator implementation
+	virtual void CalculateResult(FProgressCancel* Progress) override
+	{
+		if (PlaneTransforms.IsEmpty())
+		{
+			return;
+		}
+
+		UE::Geometry::FDynamicMeshEditor MeshEditor(ResultMesh.Get());
+		FVector Center(0, 0, 0);
+		for (const FTransform& Transform : PlaneTransforms)
+		{
+			Center += Transform.GetLocation();
+		}
+		Center /= (double)PlaneTransforms.Num();
+		UE::Geometry::FTransformSRT3d CenteringTransform(Center);
+		SetResultTransform(CenteringTransform);
+
+		for (int32 PlaneIdx = 0; PlaneIdx < PlaneTransforms.Num(); ++PlaneIdx)
+		{
+			UE::Geometry::FRectangleMeshGenerator RectangleGenerator;
+			RectangleGenerator.Width = PlaneSize;
+			RectangleGenerator.Height = PlaneSize;
+			
+			constexpr int32 MaxCountPerDim = 2000;
+			RectangleGenerator.WidthVertexCount = FMath::Min(MaxCountPerDim, RectangleGenerator.Width / NoiseSettings.PointSpacing);
+			RectangleGenerator.HeightVertexCount = FMath::Min(MaxCountPerDim, RectangleGenerator.Height / NoiseSettings.PointSpacing);
+
+			if (Progress && Progress->Cancelled())
+			{
+				return;
+			}
+
+			FDynamicMesh3 RectMesh(&RectangleGenerator.Generate());
+			const FTransform& PlaneTransform = PlaneTransforms[PlaneIdx];
+			const FVector& NoisePivot = NoisePivots[PlaneIdx];
+			const FNoiseOffsets& NoiseOffset = NoiseOffsets[PlaneIdx];
+			const FVector& ExplodedVector = ExplodedVectors[PlaneIdx];
+			FVector Normal = PlaneTransform.GetUnitAxis(EAxis::Z);
+			for (int VID : RectMesh.VertexIndicesItr())
+			{
+				FVector3d WorldPos = PlaneTransform.TransformPosition(RectMesh.GetVertex(VID));
+				FVector3d NewWorldPos = WorldPos + NoiseSettings.NoiseVector(WorldPos - NoisePivot, NoiseOffset).Dot(Normal) * Normal;
+				RectMesh.SetVertex(VID, NewWorldPos + ExplodedVector - Center);
+			}
+			UE::Geometry::FMeshIndexMappings Mappings;
+			MeshEditor.AppendMesh(&RectMesh, Mappings);
+		}
+		UE::Geometry::FMeshNormals::QuickRecomputeOverlayNormals(*ResultMesh);
+	}
+};
+
+
+TUniquePtr<UE::Geometry::FDynamicMeshOperator> UFractureToolPlaneCut::MakeNewOperator()
+{
+	UFractureSettings* FractureSettings = GetMutableDefault<UFractureSettings>();
+	NoisePreviewExplodeAmount = FractureSettings->ExplodeAmount;
+
+	TUniquePtr<FPlaneNoisePreviewOp> NoisePreviewOp = MakeUnique<FPlaneNoisePreviewOp>();
+	CutterSettings->TransferNoiseSettings(NoisePreviewOp->NoiseSettings);
+
+	if (GizmoSettings->IsGizmoEnabled()) // set preview from gizmo
+	{
+		NoisePreviewOp->PlaneSize = GizmoPlaneSize;
+		const FTransform& Transform = GizmoSettings->GetTransform();
+		EnumerateVisualizationMapping(PlanesMappings, RenderCuttingPlanesTransforms.Num(), [&](int32 Idx, FVector ExplodedVector)
+			{
+				NoisePreviewOp->PlaneTransforms.Add(Transform);
+				NoisePreviewOp->ExplodedVectors.Add(ExplodedVector);
+				NoisePreviewOp->NoiseOffsets.Add(NoiseOffsets[Idx]);
+				NoisePreviewOp->NoisePivots.Add(NoisePivots[Idx]);
+			});
+	}
+	else // set from computed transforms
+	{
+		NoisePreviewOp->PlaneSize = RenderCuttingPlaneSize;
+		EnumerateVisualizationMapping(PlanesMappings, RenderCuttingPlanesTransforms.Num(), [&](int32 Idx, FVector ExplodedVector)
+			{
+				const FTransform& Transform = RenderCuttingPlanesTransforms[Idx];
+				NoisePreviewOp->PlaneTransforms.Add(Transform);
+				NoisePreviewOp->ExplodedVectors.Add(ExplodedVector);
+				NoisePreviewOp->NoiseOffsets.Add(NoiseOffsets[Idx]);
+				NoisePreviewOp->NoisePivots.Add(NoisePivots[Idx]);
+			});
+	}
+	NoisePreviewOp->PlaneSize *= PlaneCutSettings->NoisePreviewScale;
+
+	return NoisePreviewOp;
+}
+
 
 UFractureToolPlaneCut::UFractureToolPlaneCut(const FObjectInitializer& ObjInit) 
 	: Super(ObjInit) 
@@ -35,6 +146,27 @@ void UFractureToolPlaneCut::Setup()
 	GizmoSettings->Setup(this);
 	PlaneCutSettings->bCanCutWithMultiplePlanes = !GizmoSettings->bUseGizmo;
 	NotifyOfPropertyChangeByTool(PlaneCutSettings);
+
+	// Initialize the background compute object for the noise preview
+	if (GEditor && !NoisePreview)
+	{
+		UWorld* World = GEditor->GetEditorWorldContext().World();
+		NoisePreview = NewObject<UMeshOpPreviewWithBackgroundCompute>(this);
+		NoisePreview->Setup(World, this);
+
+		TArray<UMaterialInterface*> Materials;
+		// Note the ToolSetupUtil functions optionally take a UInteractiveToolManager to provide fallback materials,
+		// but fracture mode does not set that up, so they may return a null material.
+		// Note that if the working material is null, the preview will just always use the default material instead.
+		UMaterialInterface* NoiseMaterial = ToolSetupUtil::GetDefaultSculptMaterial(nullptr);
+		if (!NoiseMaterial)
+		{
+			NoiseMaterial = ToolSetupUtil::GetDefaultMaterial();
+		}
+		UMaterialInterface* WorkingMaterial = ToolSetupUtil::GetDefaultWorkingMaterial(nullptr);
+		Materials.Add(NoiseMaterial);
+		NoisePreview->ConfigureMaterials(Materials, WorkingMaterial);
+	}
 }
 
 
@@ -42,6 +174,12 @@ void UFractureToolPlaneCut::Shutdown()
 {
 	Super::Shutdown();
 	GizmoSettings->Shutdown();
+
+	if (NoisePreview)
+	{
+		NoisePreview->Shutdown();
+		NoisePreview = nullptr;
+	}
 }
 
 
@@ -68,6 +206,15 @@ void UFractureToolPlaneCut::RegisterUICommand( FFractureEditorCommands* BindingC
 
 void UFractureToolPlaneCut::Render(const FSceneView* View, FViewport* Viewport, FPrimitiveDrawInterface* PDI)
 {
+	// Note: ExplodeAmount changes generally do not trigger tool context changes, but can affect visualization
+	// So we detect such changes here and invalidate the preview
+	// We could instead try to translate the existing NoisePreview, but that would require a separate NoisePreview per plane
+	UFractureSettings* FractureSettings = GetMutableDefault<UFractureSettings>();
+	if (NoisePreview && NoisePreviewExplodeAmount != FractureSettings->ExplodeAmount)
+	{
+		NoisePreview->InvalidateResult();
+	}
+
 	const UFracturePlaneCutSettings* LocalCutSettings = PlaneCutSettings;
 	if (CutterSettings->bDrawDiagram)
 	{
@@ -104,7 +251,7 @@ void UFractureToolPlaneCut::Render(const FSceneView* View, FViewport* Viewport,
 			const FTransform& Transform = GizmoSettings->GetTransform();
 			EnumerateVisualizationMapping(PlanesMappings, RenderCuttingPlanesTransforms.Num(), [&](int32 Idx, FVector ExplodedVector)
 			{
-				DrawPlane(Transform, 100.f, ExplodedVector);
+				DrawPlane(Transform, GizmoPlaneSize, ExplodedVector);
 			});
 		}
 		else // draw from computed transforms
@@ -141,6 +288,31 @@ void UFractureToolPlaneCut::PostEditChangeChainProperty(FPropertyChangedChainEve
 }
 
 
+void UFractureToolPlaneCut::ClearVisualizations()
+{
+	Super::ClearVisualizations();
+	RenderCuttingPlanesTransforms.Empty();
+	NoiseOffsets.Empty();
+	NoisePivots.Empty();
+	PlanesMappings.Empty();
+	if (NoisePreview)
+	{
+		NoisePreview->InvalidateResult();
+	}
+}
+
+
+void UFractureToolPlaneCut::OnTick(float DeltaTime)
+{
+	if (NoisePreview)
+	{
+		NoisePreview->Tick(DeltaTime);
+	}
+
+	Super::OnTick(DeltaTime);
+}
+
+
 void UFractureToolPlaneCut::FractureContextChanged()
 {
 	UpdateDefaultRandomSeed();
@@ -148,6 +320,11 @@ void UFractureToolPlaneCut::FractureContextChanged()
 
 	ClearVisualizations();
 
+	if (NoisePreview)
+	{
+		NoisePreview->SetVisibility(PlaneCutSettings->bShowNoisePreview);
+	}
+
 	RenderCuttingPlaneSize = FLT_MAX;
 	for (FFractureToolContext& FractureContext : FractureContexts)
 	{
@@ -158,15 +335,28 @@ void UFractureToolPlaneCut::FractureContextChanged()
 		}
 		int32 CollectionIdx = VisualizedCollections.Add(FractureContext.GetGeometryCollectionComponent());
 		int32 BoneIdx = FractureContext.GetSelection().Num() == 1 ? FractureContext.GetSelection()[0] : INDEX_NONE;
-		PlanesMappings.AddMapping(CollectionIdx, BoneIdx, RenderCuttingPlanesTransforms.Num());
+		int32 TransformsStart = RenderCuttingPlanesTransforms.Num();
+		PlanesMappings.AddMapping(CollectionIdx, BoneIdx, TransformsStart);
 
 		GenerateSliceTransforms(FractureContext, RenderCuttingPlanesTransforms);
+		for (int32 Idx = TransformsStart; Idx < RenderCuttingPlanesTransforms.Num(); ++Idx)
+		{
+			FRandomStream RandomStream(FractureContext.GetSeed());
+			FNoiseOffsets NoiseOffset(RandomStream);
+			NoiseOffsets.Add(NoiseOffset);
+			NoisePivots.Add(FractureContext.GetTransform().GetLocation());
+		}
 
 		if (Bounds.GetExtent().GetMax() < RenderCuttingPlaneSize)
 		{
 			RenderCuttingPlaneSize = Bounds.GetExtent().GetMax();
 		}
 	}
+
+	if (NoisePreview)
+	{
+		NoisePreview->InvalidateResult();
+	}
 }
 
 
diff --git a/Engine/Plugins/Experimental/ChaosEditor/Source/FractureEditor/Private/FractureToolPlaneCut.h b/Engine/Plugins/Experimental/ChaosEditor/Source/FractureEditor/Private/FractureToolPlaneCut.h
--- a/Engine/Plugins/Experimental/ChaosEditor/Source/FractureEditor/Private/FractureToolPlaneCut.h
+++ b/Engine/Plugins/Experimental/ChaosEditor/Source/FractureEditor/Private/FractureToolPlaneCut.h
@@ -3,11 +3,14 @@
 #pragma once
 
 #include "FractureToolCutter.h"
+#include "ModelingOperators.h"
 
 #include "FractureToolPlaneCut.generated.h"
 
 class FFractureToolContext;
-
+class UMeshOpPreviewWithBackgroundCompute;
+struct FNoiseOffsets;
+struct FNoiseSettings;
 
 UCLASS(config = EditorPerProjectUserSettings)
 class UFracturePlaneCutSettings : public UFractureToolSettings
@@ -25,11 +28,20 @@ public:
 
 	UPROPERTY()
 	bool bCanCutWithMultiplePlanes = false;
+
+	/** Whether to show the noise displacement on the preview visualization of the cutting planes */
+	UPROPERTY(EditAnywhere, Category = Visualization)
+	bool bShowNoisePreview = true;
+
+	/** Scale of the noise preview plane */
+	UPROPERTY(EditAnywhere, Category = Visualization, meta = (ClampMin = ".01", ClampMax = "1000", UIMin = ".25", UIMax = "10", EditCondition = "bShowNoisePreview", HideEditConditionToggle))
+	double NoisePreviewScale = 1.;
+
 };
 
 
 UCLASS(DisplayName = "Plane Cut Tool", Category = "FractureTools")
-class UFractureToolPlaneCut : public UFractureToolCutterBase
+class UFractureToolPlaneCut : public UFractureToolCutterBase, public UE::Geometry::IDynamicMeshOperatorFactory
 {
 public:
 	GENERATED_BODY()
@@ -42,7 +54,9 @@ public:
 	virtual FSlateIcon GetToolIcon() const override;
 	virtual void SelectedBonesChanged() override;
 
-	void Render(const FSceneView* View, FViewport* Viewport, FPrimitiveDrawInterface* PDI) override;
+
+	virtual void OnTick(float DeltaTime) override;
+	virtual void Render(const FSceneView* View, FViewport* Viewport, FPrimitiveDrawInterface* PDI) override;
 
 	virtual void RegisterUICommand(FFractureEditorCommands* BindingContext) override;
 
@@ -54,13 +68,11 @@ public:
 	virtual void Setup() override;
 	virtual void Shutdown() override;
 
+	// IDynamicMeshOperatorFactory API, for generating noise preview meshes
+	virtual TUniquePtr<UE::Geometry::FDynamicMeshOperator> MakeNewOperator() override;
+
 protected:
-	virtual void ClearVisualizations() override
-	{
-		Super::ClearVisualizations();
-		RenderCuttingPlanesTransforms.Empty();
-		PlanesMappings.Empty();
-	}
+	virtual void ClearVisualizations() override;
 
 	virtual void PostEditChangeChainProperty(struct FPropertyChangedChainEvent& PropertyChangedEvent) override;
 
@@ -72,11 +84,18 @@ private:
 	UPROPERTY(EditAnywhere, Category = Uniform)
 	TObjectPtr<UFractureTransformGizmoSettings> GizmoSettings;
 
+	UPROPERTY()
+	TObjectPtr<UMeshOpPreviewWithBackgroundCompute> NoisePreview;
+
 	void GenerateSliceTransforms(const FFractureToolContext& Context, TArray<FTransform>& CuttingPlaneTransforms);
 
 	float RenderCuttingPlaneSize;
+	const float GizmoPlaneSize = 100.f;
 	TArray<FTransform> RenderCuttingPlanesTransforms;
+	TArray<FNoiseOffsets> NoiseOffsets;
+	TArray<FVector> NoisePivots;
 	FVisualizationMappings PlanesMappings;
+	float NoisePreviewExplodeAmount = 0;
 };
 
 
diff --git a/Engine/Plugins/Experimental/ChaosEditor/Source/FractureEditor/Public/FractureEditorMode.h b/Engine/Plugins/Experimental/ChaosEditor/Source/FractureEditor/Public/FractureEditorMode.h
--- a/Engine/Plugins/Experimental/ChaosEditor/Source/FractureEditor/Public/FractureEditorMode.h
+++ b/Engine/Plugins/Experimental/ChaosEditor/Source/FractureEditor/Public/FractureEditorMode.h
@@ -34,7 +34,7 @@ public:
 	virtual void Enter() override;
 	virtual void Exit() override;
 
-	//virtual void Tick(FEditorViewportClient* ViewportClient, float DeltaTime) override;
+	virtual void Tick(FEditorViewportClient* ViewportClient, float DeltaTime) override;
 	virtual void Render(const FSceneView* View, FViewport* Viewport, FPrimitiveDrawInterface* PDI) override;
 	virtual void DrawHUD(FEditorViewportClient* ViewportClient, FViewport* Viewport, const FSceneView* View, FCanvas* Canvas) override;
 
@@ -79,8 +79,11 @@ private:
 	/** Handle package reloading (might be our geometry collection) */
 	void HandlePackageReloaded(const EPackageReloadPhase InPackageReloadPhase, FPackageReloadedEvent* InPackageReloadedEvent);
 
-	static FConvexVolume TranformFrustum(const FConvexVolume& InFrustum, const FMatrix& InMatrix);
+	static FConvexVolume TransformFrustum(const FConvexVolume& InFrustum, const FMatrix& InMatrix);
 	static FConvexVolume GetVolumeFromBox(const FBox &InBox);
+
+	void ConfigureRealTimeViewportsOverride(bool bEnable);
+
 private:
 	/** This selection set is updated from actor selection changed event.  We change state on components as they are selected so we have to maintain or own list **/
 	TArray<UGeometryCollectionComponent*> SelectedGeometryComponents;
diff --git a/Engine/Plugins/Experimental/ChaosEditor/Source/FractureEditor/Public/FractureTool.h b/Engine/Plugins/Experimental/ChaosEditor/Source/FractureEditor/Public/FractureTool.h
--- a/Engine/Plugins/Experimental/ChaosEditor/Source/FractureEditor/Public/FractureTool.h
+++ b/Engine/Plugins/Experimental/ChaosEditor/Source/FractureEditor/Public/FractureTool.h
@@ -159,6 +159,8 @@ public:
 	/** Executes function that generates new geometry. Returns the first new geometry index. */
 	virtual int32 ExecuteFracture(const FFractureToolContext& FractureContext) { return INDEX_NONE; }
 
+	virtual void OnTick(float DeltaTime) {}
+
 	/** Draw callback from edmode*/
 	virtual void Render(const FSceneView* View, FViewport* Viewport, FPrimitiveDrawInterface* PDI) {}
 	virtual void DrawHUD(FEditorViewportClient* ViewportClient, FViewport* Viewport, const FSceneView* View, FCanvas* Canvas) {}
diff --git a/Engine/Plugins/Experimental/PlanarCutPlugin/Source/PlanarCut/Public/PlanarCut.h b/Engine/Plugins/Experimental/PlanarCutPlugin/Source/PlanarCut/Public/PlanarCut.h
--- a/Engine/Plugins/Experimental/PlanarCutPlugin/Source/PlanarCut/Public/PlanarCut.h
+++ b/Engine/Plugins/Experimental/PlanarCutPlugin/Source/PlanarCut/Public/PlanarCut.h
@@ -11,6 +11,29 @@
 
 class FProgressCancel;
 
+// Offsets to randomize where Perlin noise is sampled for each dimension of a noise vector
+struct FNoiseOffsets
+{
+	FNoiseOffsets(FRandomStream& RandomStream) 
+		: X(RandomStream.VRand() * RandomOffsetScale), Y(RandomStream.VRand() * RandomOffsetScale), Z(RandomStream.VRand() * RandomOffsetScale)
+	{
+	}
+
+	FNoiseOffsets() {}
+
+	void SetOffsets(FRandomStream& RandomStream)
+	{
+		X = RandomStream.VRand() * RandomOffsetScale;
+		Y = RandomStream.VRand() * RandomOffsetScale;
+		Z = RandomStream.VRand() * RandomOffsetScale;
+	}
+
+	FVector X, Y, Z;
+private:
+	// Offset scale chosen to keep reasonably small offsets while still de-correlating the noise in each dimension.
+	static constexpr double RandomOffsetScale = 100.0;
+};
+
 struct PLANARCUT_API FNoiseSettings
 {
 	float Amplitude = 2;
@@ -19,6 +42,28 @@ struct PLANARCUT_API FNoiseSettings
 	float PointSpacing = 1;
 	float Lacunarity = 2;
 	float Persistence = .5;
+
+	FVector NoiseVector(const FVector& Pos, const FNoiseOffsets& Offsets) const
+	{
+		FVector Base = Pos * Frequency;
+		return FVector(
+			OctaveNoise(Base + Offsets.X),
+			OctaveNoise(Base + Offsets.Y),
+			OctaveNoise(Base + Offsets.Z)
+		) * Amplitude;
+	}
+
+	float OctaveNoise(const FVector& V) const
+	{
+		float NoiseValue = 0;
+		float FreqScale = 1;
+		float AmpScale = 1;
+		for (int32 Octave = 0; Octave < Octaves; Octave++, FreqScale *= Lacunarity, AmpScale *= Persistence)
+		{
+			NoiseValue += FMath::PerlinNoise3D(V * FreqScale) * AmpScale;
+		}
+		return NoiseValue;
+	}
 };
 
 // auxiliary structure for FPlanarCells to carry material info
