Fracture editor : fix uninitialized number of embedded element displayed when no geometry collection is selected

diff --git a/Engine/Plugins/Experimental/ChaosEditor/Source/FractureEditor/Private/SGeometryCollectionStatistics.h b/Engine/Plugins/Experimental/ChaosEditor/Source/FractureEditor/Private/SGeometryCollectionStatistics.h
--- a/Engine/Plugins/Experimental/ChaosEditor/Source/FractureEditor/Private/SGeometryCollectionStatistics.h
+++ b/Engine/Plugins/Experimental/ChaosEditor/Source/FractureEditor/Private/SGeometryCollectionStatistics.h
@@ -10,7 +10,7 @@
 struct FGeometryCollectionStatistics
 {
 	TArray<uint32> CountsPerLevel;
-	uint32 EmbeddedCount;
+	uint32 EmbeddedCount = 0;
 };
 
 class SGeometryCollectionStatistics : public SCompoundWidget
