Fracture editor : refactor geometry collection outliner to simplify adding new columns and data to it
- remove the need for Guids to identify items ( also in the histogram control )
- items are now fetching their data in a single data collection owned by the root parent item that only fetches the data required to display the columns displayed

diff --git a/Engine/Plugins/Experimental/ChaosEditor/Source/FractureEditor/Private/FractureEditorModeToolkit.cpp b/Engine/Plugins/Experimental/ChaosEditor/Source/FractureEditor/Private/FractureEditorModeToolkit.cpp
--- a/Engine/Plugins/Experimental/ChaosEditor/Source/FractureEditor/Private/FractureEditorModeToolkit.cpp
+++ b/Engine/Plugins/Experimental/ChaosEditor/Source/FractureEditor/Private/FractureEditorModeToolkit.cpp
@@ -713,10 +713,6 @@ void FFractureEditorModeToolkit::OnObjectPostEditChange( UObject* Object, FPrope
 		{
 			HistogramView->RegenerateNodes(GetLevelViewValue());
 		}
-		else if (PropertyChangedEvent.Property->GetFName() == GET_MEMBER_NAME_CHECKED(UOutlinerSettings, ItemText))
-		{
-			OutlinerView->RegenerateItems();
-		}
 		else if (PropertyChangedEvent.Property->GetFName() == GET_MEMBER_NAME_CHECKED(UOutlinerSettings, ColorByLevel))
 		{
 			OutlinerView->RegenerateItems();
@@ -727,7 +723,6 @@ void FFractureEditorModeToolkit::OnObjectPostEditChange( UObject* Object, FPrope
 		else if (PropertyChangedEvent.Property->GetFName() == GET_MEMBER_NAME_CHECKED(UOutlinerSettings, ColumnMode))
 		{
 			OutlinerView->RegenerateHeader();
-			OutlinerView->RegenerateItems();
 			FGeometryCollectionStatistics Stats;
 			GetStatisticsSummary(Stats);
 			StatisticsView->SetStatistics(Stats);
diff --git a/Engine/Plugins/Experimental/ChaosEditor/Source/FractureEditor/Private/SGeometryCollectionHistogram.cpp b/Engine/Plugins/Experimental/ChaosEditor/Source/FractureEditor/Private/SGeometryCollectionHistogram.cpp
--- a/Engine/Plugins/Experimental/ChaosEditor/Source/FractureEditor/Private/SGeometryCollectionHistogram.cpp
+++ b/Engine/Plugins/Experimental/ChaosEditor/Source/FractureEditor/Private/SGeometryCollectionHistogram.cpp
@@ -114,15 +114,7 @@ void FGeometryCollectionHistogramItem::SetNormalizedValue(float MinValue, float
 
 FGeometryCollectionHistogramItemPtr FGeometryCollectionHistogramItemComponent::GetItemFromBoneIndex(int32 BoneIndex) const
 {
-	for (auto& Pair : NodesMap)
-	{
-		if (Pair.Value->GetBoneIndex() == BoneIndex)
-		{
-			return Pair.Value;
-		}
-	}
-
-	return FGeometryCollectionHistogramItemPtr();
+	return ItemByBoneIndex.FindRef(BoneIndex);
 }
 
 FGeometryCollectionHistogramItemList FGeometryCollectionHistogramItemComponent::RegenerateNodes(int32 LevelView)
@@ -142,15 +134,12 @@ FGeometryCollectionHistogramItemList FGeometryCollectionHistogramItemComponent::
 	{
 		FGeometryCollection* Collection = Component->GetRestCollection()->GetGeometryCollection().Get();
 
-		NodesMap.Empty();
-		GuidIndexMap.Empty();
+		ItemByBoneIndex.Empty();
 
 		if (Collection)
 		{
 			int32 NumElements = Collection->NumElements(FGeometryCollection::TransformGroup);
-			::GeometryCollection::GenerateTemporaryGuids(Collection);
 
-			const TManagedArray<FGuid>& Guids = Collection->GetAttribute<FGuid>("GUID", "Transform");
 			const TManagedArray<int32>& GeometryToTransform = Collection->GetAttribute<int32>("TransformIndex", "Geometry");
 			const TManagedArray<FLinearColor>& BoneColor = Collection->GetAttribute<FLinearColor>("BoneColor", "Transform");
 		
@@ -172,14 +161,13 @@ FGeometryCollectionHistogramItemList FGeometryCollectionHistogramItemComponent::
 						}
 					}
 
-					TSharedRef<FGeometryCollectionHistogramItem> NewItem = MakeShared<FGeometryCollectionHistogramItem>(Guids[Index], Index, AsShared());
+					TSharedRef<FGeometryCollectionHistogramItem> NewItem = MakeShared<FGeometryCollectionHistogramItem>(Index, AsShared());
 
-					FLinearColor GeoColor = BoneColor[Index];
+					const FLinearColor GeoColor = BoneColor[Index];
 					NewItem->SetColor(GeoColor);
 
 					NodesList.Add(NewItem);
-					NodesMap.Add(Guids[Index], NewItem);
-					GuidIndexMap.Add(Guids[Index], Index);
+					ItemByBoneIndex.Add(Index, NewItem);
 				}
 			}
 		}
diff --git a/Engine/Plugins/Experimental/ChaosEditor/Source/FractureEditor/Private/SGeometryCollectionHistogram.h b/Engine/Plugins/Experimental/ChaosEditor/Source/FractureEditor/Private/SGeometryCollectionHistogram.h
--- a/Engine/Plugins/Experimental/ChaosEditor/Source/FractureEditor/Private/SGeometryCollectionHistogram.h
+++ b/Engine/Plugins/Experimental/ChaosEditor/Source/FractureEditor/Private/SGeometryCollectionHistogram.h
@@ -79,8 +79,7 @@ public:
 
 private:
 	TWeakObjectPtr<UGeometryCollectionComponent> Component;
-	TMap<FGuid, FGeometryCollectionHistogramItemPtr> NodesMap;
-	TMap<FGuid, int32> GuidIndexMap;
+	TMap<int32, FGeometryCollectionHistogramItemPtr> ItemByBoneIndex;
 };
 
 
@@ -88,9 +87,8 @@ private:
 class FGeometryCollectionHistogramItem : public TSharedFromThis<FGeometryCollectionHistogramItem>
 {
 public:
-	FGeometryCollectionHistogramItem(const FGuid NewGuid, const int32 InBoneIndex, const TSharedPtr<FGeometryCollectionHistogramItemComponent> InParentComponentItem)
-		: Guid(NewGuid)
-		, BoneIndex(InBoneIndex)
+	FGeometryCollectionHistogramItem(const int32 InBoneIndex, const TSharedPtr<FGeometryCollectionHistogramItemComponent> InParentComponentItem)
+		: BoneIndex(InBoneIndex)
 		, ParentComponentItem(InParentComponentItem)
 		, ListIndex(0)
 		, NodeColor(FLinearColor::Black)
@@ -112,7 +110,6 @@ public:
 	int32 GetListIndex() const { return ListIndex; }
 
 private:
-	const FGuid Guid;
 	const int32 BoneIndex;
 	const TSharedPtr<FGeometryCollectionHistogramItemComponent> ParentComponentItem;
 	
diff --git a/Engine/Plugins/Experimental/ChaosEditor/Source/FractureEditor/Private/SGeometryCollectionOutliner.cpp b/Engine/Plugins/Experimental/ChaosEditor/Source/FractureEditor/Private/SGeometryCollectionOutliner.cpp
--- a/Engine/Plugins/Experimental/ChaosEditor/Source/FractureEditor/Private/SGeometryCollectionOutliner.cpp
+++ b/Engine/Plugins/Experimental/ChaosEditor/Source/FractureEditor/Private/SGeometryCollectionOutliner.cpp
@@ -185,16 +185,15 @@ FReply FGeometryCollectionTreeItemBone::OnDrop(const FDragDropEvent& DragDropEve
 
 UOutlinerSettings::UOutlinerSettings(const FObjectInitializer& ObjInit)
 	: Super(ObjInit)
-	, ItemText(EOutlinerItemNameEnum::BoneIndex)
 	, ColorByLevel(false)
-	, ColumnMode(EOutlinerColumnMode::StateAndSize)
+	, ColumnMode(EOutlinerColumnMode::State)
 {}
 
 TSharedRef<SWidget> SGeometryCollectionOutlinerRow::GenerateWidgetForColumn(const FName& ColumnName)
 {
-	if (ColumnName == SGeometryCollectionOutlinerColumnID::Bone)
+	if (ColumnName == SGeometryCollectionOutlinerColumnID::BoneIndex)
 	{
-		const TSharedPtr<SWidget> NameWidget = Item->MakeNameColumnWidget();
+		const TSharedPtr<SWidget> NameWidget = Item->MakeBoneIndexColumnWidget();
 		return SNew(SHorizontalBox)
 			+ SHorizontalBox::Slot()
 				.HAlign(HAlign_Left)
@@ -209,6 +208,8 @@ TSharedRef<SWidget> SGeometryCollectionOutlinerRow::GenerateWidgetForColumn(cons
 					NameWidget.ToSharedRef()
 				];
 	}
+	if (ColumnName == SGeometryCollectionOutlinerColumnID::BoneName)
+		return Item->MakeBoneNameColumnWidget();
 	if (ColumnName == SGeometryCollectionOutlinerColumnID::RelativeSize)
 		return Item->MakeRelativeSizeColumnWidget();
 	if (ColumnName == SGeometryCollectionOutlinerColumnID::Volume)
@@ -265,9 +266,16 @@ void SGeometryCollectionOutliner::RegenerateHeader()
 
 	HeaderRowWidget->ClearColumns();
 
+	RegenerateRootData();
+
+	HeaderRowWidget->AddColumn(
+		SHeaderRow::Column(SGeometryCollectionOutlinerColumnID::BoneIndex)
+		.DefaultLabel(LOCTEXT("GCOutliner_Column_BoneIndex", "Index"))
+		.FillWidth(2.0f)
+	);
 	HeaderRowWidget->AddColumn(
-		SHeaderRow::Column(SGeometryCollectionOutlinerColumnID::Bone)
-		.DefaultLabel(LOCTEXT("GCOutliner_Column_Bone", "Bone"))
+		SHeaderRow::Column(SGeometryCollectionOutlinerColumnID::BoneName)
+		.DefaultLabel(LOCTEXT("GCOutliner_Column_BoneName", "Name"))
 		.FillWidth(2.0f)
 	);
 
@@ -275,21 +283,7 @@ void SGeometryCollectionOutliner::RegenerateHeader()
 	UOutlinerSettings* OutlinerSettings = GetMutableDefault<UOutlinerSettings>();
 	switch (OutlinerSettings->ColumnMode)
 	{
-	case EOutlinerColumnMode::StateAndSize:
-		HeaderRowWidget->AddColumn(
-			SHeaderRow::Column(SGeometryCollectionOutlinerColumnID::RelativeSize)
-				.DefaultLabel(LOCTEXT("GCOutliner_Column_RelativeSize", "Relative Size"))
-				.DefaultTooltip(LOCTEXT("GCOutliner_Column_RelativeSize_ToolTip", "Relative size ( Used for size specific data )"))
-				.HAlignHeader(EHorizontalAlignment::HAlign_Center)
-				.FillWidth(CustomFillWidth)
-		);
-		HeaderRowWidget->AddColumn(
-			SHeaderRow::Column(SGeometryCollectionOutlinerColumnID::Volume)
-				.DefaultLabel(LOCTEXT("GCOutliner_Column_Volume", "Size"))
-				.DefaultTooltip(LOCTEXT("GCOutliner_Column_Volume_ToolTip", "Side length of the bounding cube (in cm)"))
-				.HAlignHeader(EHorizontalAlignment::HAlign_Center)
-				.FillWidth(CustomFillWidth)
-		);
+	case EOutlinerColumnMode::State:
 		HeaderRowWidget->AddColumn(
 			SHeaderRow::Column(SGeometryCollectionOutlinerColumnID::InitialState)
 				.DefaultLabel(LOCTEXT("GCOutliner_Column_InitialState", "Initial State"))
@@ -355,14 +349,43 @@ void SGeometryCollectionOutliner::RegenerateHeader()
 				.FillWidth(CustomFillWidth)
 		);
 		break;
+
+	case EOutlinerColumnMode::Size:
+		HeaderRowWidget->AddColumn(
+			SHeaderRow::Column(SGeometryCollectionOutlinerColumnID::RelativeSize)
+			.DefaultLabel(LOCTEXT("GCOutliner_Column_RelativeSize", "Relative Size"))
+			.DefaultTooltip(LOCTEXT("GCOutliner_Column_RelativeSize_ToolTip", "Relative size ( Used for size specific data )"))
+			.HAlignHeader(EHorizontalAlignment::HAlign_Center)
+			.FillWidth(CustomFillWidth)
+		);
+		HeaderRowWidget->AddColumn(
+			SHeaderRow::Column(SGeometryCollectionOutlinerColumnID::Volume)
+			.DefaultLabel(LOCTEXT("GCOutliner_Column_Volume", "Size"))
+			.DefaultTooltip(LOCTEXT("GCOutliner_Column_Volume_ToolTip", "Side length of the bounding cube (in cm)"))
+			.HAlignHeader(EHorizontalAlignment::HAlign_Center)
+			.FillWidth(CustomFillWidth)
+		);
+		break;
 	}
 }
 
 void SGeometryCollectionOutliner::RegenerateItems()
 {
+	RegenerateRootData();
 	TreeView->RebuildList();
 }
 
+void SGeometryCollectionOutliner::RegenerateRootData()
+{
+	for (TSharedPtr<FGeometryCollectionTreeItemComponent> RootNode : RootNodes)
+	{
+		if (RootNode)
+		{
+			RootNode->GenerateDataCollection();
+		}
+	}
+}
+
 TSharedRef<ITableRow> SGeometryCollectionOutliner::MakeTreeRowWidget(FGeometryCollectionTreeItemPtr InItem, const TSharedRef<STableViewBase>& InOwnerTable)
 {
 	return InItem->MakeTreeRowWidget(InOwnerTable);
@@ -579,6 +602,197 @@ void SGeometryCollectionOutliner::SetAnchored(bool bAnchored)
 	RegenerateItems();
 }
 
+///////////////////////////////////////////////////////////////////////////////////////////////////////////////
+
+// for now use the transform group as this simplify the attribute copy 
+static const FName DataCollectionGroup = FGeometryCollection::TransformGroup;
+
+FGeometryCollectionItemDataFacade::FGeometryCollectionItemDataFacade(FManagedArrayCollection& InCollection)
+	: DataCollection(InCollection)
+	, BoneNameAttribute(InCollection, "BoneName", DataCollectionGroup)
+	, LevelAttribute(InCollection, "Level", DataCollectionGroup)
+	, VisibleAttribute(InCollection, "Visible", DataCollectionGroup)
+	, InitialStateAttribute(InCollection, "InitialState", DataCollectionGroup)
+	, RelativeSizeAttribute(InCollection, "RelativeSize", DataCollectionGroup)
+	, VolumeAttribute(InCollection, "Volume", DataCollectionGroup)
+	, AnchoredAttribute(InCollection, "Anchored", DataCollectionGroup)
+	, DamageAttribute(InCollection, "Damage", DataCollectionGroup)
+	, DamageThresholdAttribute(InCollection, "DamageThreshold", DataCollectionGroup)
+	, BrokenStateAttribute(InCollection, "BrokenState", DataCollectionGroup)
+	, RemoveOnBreakAttribute(InCollection, "RemoveOnBreak", DataCollectionGroup)
+	, SimulationTypeAttribute(InCollection, "SimulationType", DataCollectionGroup)
+	, HasSourceCollisionAttribute(InCollection, "HasSourceCollision", DataCollectionGroup)
+	, SourceCollisionUsedAttribute(InCollection, "SourceCollisionUsed", DataCollectionGroup)
+{
+}
+
+template <typename T>
+static T GetAttributeValue(const TManagedArrayAccessor<T>& Attribute, int32 Index, T Default)
+{
+	return (Attribute.IsValid()) ? Attribute.Get()[Index] : Default;
+}
+
+FString FGeometryCollectionItemDataFacade::GetBoneName(int32 Index) const
+{
+	return GetAttributeValue(BoneNameAttribute, Index, FString());
+}
+
+float FGeometryCollectionItemDataFacade::GetRelativeSize(int32 Index) const
+{
+	return GetAttributeValue(RelativeSizeAttribute, Index, 0.0f);
+}
+
+float FGeometryCollectionItemDataFacade::GetVolume(int32 Index) const
+{
+	return GetAttributeValue(VolumeAttribute, Index, 0.0f);
+}
+int32 FGeometryCollectionItemDataFacade::GetInitialState(int32 Index) const
+{
+	return GetAttributeValue(InitialStateAttribute, Index, 0);
+}
+bool FGeometryCollectionItemDataFacade::IsAnchored(int32 Index) const
+{
+	return GetAttributeValue(AnchoredAttribute, Index, false);
+}
+float FGeometryCollectionItemDataFacade::GetDamage(int32 Index) const
+{
+	return GetAttributeValue(DamageAttribute, Index, 0.0f);
+}
+float FGeometryCollectionItemDataFacade::GetDamageThreshold(int32 Index) const
+{
+	return GetAttributeValue(DamageThresholdAttribute, Index, 0.0f);
+}
+bool FGeometryCollectionItemDataFacade::IsBroken(int32 Index) const
+{
+	return GetAttributeValue(BrokenStateAttribute, Index, false);
+}
+FRemoveOnBreakData FGeometryCollectionItemDataFacade::GetRemoveOnBreakData(int32 Index) const
+{
+	return GetAttributeValue(RemoveOnBreakAttribute, Index, FRemoveOnBreakData::DisabledPackedData);
+}
+bool FGeometryCollectionItemDataFacade::HasSourceCollision(int32 Index) const
+{
+	return GetAttributeValue(HasSourceCollisionAttribute, Index, false);
+}
+bool FGeometryCollectionItemDataFacade::IsSourceCollisionUsed(int32 Index) const
+{
+	return GetAttributeValue(SourceCollisionUsedAttribute, Index, false);
+}
+
+void FGeometryCollectionItemDataFacade::FillFromGeometryCollectionComponent(const UGeometryCollectionComponent& GeometryCollectionComponent, EOutlinerColumnMode ColumnMode)
+{
+	const UGeometryCollection* RestCollection = GeometryCollectionComponent.GetRestCollection();
+	if (RestCollection && IsValidChecked(RestCollection))
+	{
+		TSharedPtr<const FGeometryCollection, ESPMode::ThreadSafe> GeometryCollectionPtr = RestCollection->GetGeometryCollection();
+		const FGeometryCollection& GeometryCollection = *GeometryCollectionPtr;
+
+		const int32 GCNumElements = GeometryCollection.NumElements(FGeometryCollection::TransformGroup);
+
+		DataCollection.AddGroup(DataCollectionGroup);
+		DataCollection.AddElements(GCNumElements, DataCollectionGroup);
+
+		// bones names are part of the collection regardless of the column mode
+		TManagedArrayAccessor<FString> GCBoneNameAttribute(GeometryCollection, "BoneName", FGeometryCollection::TransformGroup);
+		if (GCBoneNameAttribute.IsValid())
+		{
+			BoneNameAttribute.Copy(GCBoneNameAttribute);
+		}
+
+		// we also need level, simulation type and visibility for item colors
+		// @todo : we copuld use the settings selection option to see if we need all of them or compute the color there
+		TManagedArrayAccessor<int32> GCLevelAttribute(GeometryCollection, "Level", FGeometryCollection::TransformGroup);
+		if (GCLevelAttribute.IsValid())
+		{
+			LevelAttribute.Copy(GCLevelAttribute);
+		}
+
+		const TManagedArrayAccessor<int32> GCSimulationTypeAttribute(GeometryCollection, "SimulationType", FGeometryCollection::TransformGroup);
+		if (GCSimulationTypeAttribute.IsValid())
+		{
+			SimulationTypeAttribute.Copy(GCSimulationTypeAttribute);
+		}
+
+		TManagedArray<bool>& Visible = VisibleAttribute.Add();
+		for (int32 Index = 0; Index < GCNumElements; Index++)
+		{
+			Visible[Index] = GeometryCollection.IsVisible(Index);
+		}
+
+		if (ColumnMode == EOutlinerColumnMode::State)
+		{
+			const TManagedArrayAccessor<int32> GCInitialStateAttribute(GeometryCollection, "InitialDynamicState", FGeometryCollection::TransformGroup);
+			if (GCInitialStateAttribute.IsValid())
+			{
+				InitialStateAttribute.Copy(GCInitialStateAttribute);
+			}
+
+			const TManagedArrayAccessor<bool> GCAnchoredAttribute(GeometryCollection, "Anchored", FGeometryCollection::TransformGroup);
+			if (GCAnchoredAttribute.IsValid())
+			{
+				AnchoredAttribute.Copy(GCAnchoredAttribute);
+			}
+		}
+		else if (ColumnMode == EOutlinerColumnMode::Size)
+		{
+			const TManagedArrayAccessor<float> GCRelativeSizeAttribute(GeometryCollection, "Size", FGeometryCollection::TransformGroup);
+			if (GCRelativeSizeAttribute.IsValid())
+			{
+				RelativeSizeAttribute.Copy(GCRelativeSizeAttribute);
+			}
+
+			const TManagedArrayAccessor<float> GCVolumeAttribute(GeometryCollection, "Volume", FGeometryCollection::TransformGroup);
+			if (GCVolumeAttribute.IsValid())
+			{
+				VolumeAttribute.Copy(GCVolumeAttribute);
+			}
+		}
+		else if (ColumnMode == EOutlinerColumnMode::Damage)
+		{
+			if (const FDamageCollector* Collector = GeometryCollectionComponent.GetRunTimeDataCollector())
+			{
+				TManagedArray<float>& Damage = DamageAttribute.Add();
+				TManagedArray<float>& DamageThreshold = DamageThresholdAttribute.Add();
+				TManagedArray<bool>& BrokenState = BrokenStateAttribute.Add();
+				for (int32 Index = 0; Index < GCNumElements; Index++)
+				{
+					const FDamageCollector::FDamageData& DamageData = (*Collector)[Index];
+					Damage[Index] = DamageData.MaxDamages;
+					DamageThreshold[Index] = DamageData.DamageThreshold;
+					BrokenState[Index] = DamageData.bIsBroken;
+				}
+			}
+		}
+		else if (ColumnMode == EOutlinerColumnMode::Collision)
+		{
+			using FImplicitGeom = FGeometryDynamicCollection::FSharedImplicit;
+			const TManagedArrayAccessor<FImplicitGeom> GCSourceCollisionAttribute(GeometryCollection, "ExternalCollisions", FGeometryCollection::TransformGroup);
+			if (GCSourceCollisionAttribute.IsValid())
+			{
+				const TManagedArray<FImplicitGeom>& GCSourceCollision = GCSourceCollisionAttribute.Get();
+				TManagedArray<bool>& HasSourceCollision = HasSourceCollisionAttribute.Add();
+				TManagedArray<bool>& SourceCollisionUsed = SourceCollisionUsedAttribute.Add();
+				
+				for (int32 Index = 0; Index < GCNumElements; Index++)
+				{
+					HasSourceCollision[Index] = (GCSourceCollision[Index] != nullptr);
+					SourceCollisionUsed[Index] = RestCollection->bImportCollisionFromSource;
+				}
+			}
+		}
+		else if (ColumnMode == EOutlinerColumnMode::Removal)
+		{
+			const TManagedArrayAccessor<FVector4f> GCRemoveOnBreakAttribute(GeometryCollection, "RemoveOnBreak", FGeometryCollection::TransformGroup);
+			if (GCRemoveOnBreakAttribute.IsValid())
+			{
+				RemoveOnBreakAttribute.Copy(GCRemoveOnBreakAttribute);
+			}
+		}
+	}
+}
+
+
+///////////////////////////////////////////////////////////////////////////////////////////////////////////////
 TSharedRef<ITableRow> FGeometryCollectionTreeItemComponent::MakeTreeRowWidget(const TSharedRef<STableViewBase>& InOwnerTable, bool bNoExtraColumn)
 {
 	FString ActorName = Component.IsValid()? Component->GetOwner()->GetActorLabel(): FString();
@@ -593,20 +807,12 @@ TSharedRef<ITableRow> FGeometryCollectionTreeItemComponent::MakeTreeRowWidget(co
 
 void FGeometryCollectionTreeItemComponent::GetChildren(FGeometryCollectionTreeItemList& OutChildren)
 {
-	OutChildren = MyChildren;
+	OutChildren = ChildItems;
 }
 
 FGeometryCollectionTreeItemPtr FGeometryCollectionTreeItemComponent::GetItemFromBoneIndex(int32 BoneIndex) const
 {
-	for (auto& Pair : NodesMap)
-	{
-		if (Pair.Value->GetBoneIndex() == BoneIndex)
-		{
-			return Pair.Value;
-		}
-	}
-
-	return FGeometryCollectionTreeItemPtr();
+	return ItemsByBoneIndex.FindRef(BoneIndex);
 }
 
 void FGeometryCollectionTreeItemComponent::GetChildrenForBone(FGeometryCollectionTreeItemBone& BoneItem, FGeometryCollectionTreeItemList& OutChildren) const
@@ -619,23 +825,13 @@ void FGeometryCollectionTreeItemComponent::GetChildrenForBone(FGeometryCollectio
 	{
 		if (FGeometryCollection* Collection = RestCollection->GetGeometryCollection().Get())
 		{
-			if (!Collection->HasAttribute("GUID", "Transform"))
-			{
-				GeometryCollection::GenerateTemporaryGuids(Collection);
-			}
-
-			if (const int32* BoneIndex = GuidIndexMap.Find(BoneItem.GetGuid()))
+			const int32 BoneIndex = BoneItem.GetBoneIndex();
+			if (ensure(BoneIndex >= 0 && BoneIndex < Collection->NumElements(FGeometryCollection::TransformGroup)))
 			{
-				if (!ensure(*BoneIndex >= 0 && *BoneIndex < Collection->NumElements(FGeometryCollection::TransformGroup)))
-				{
-					return;
-				}
 				const TManagedArray<TSet<int32>>& Children = Collection->Children;
-				const TManagedArray<FGuid>& Guids = Collection->GetAttribute<FGuid>("GUID", "Transform");
-				for (auto ChildIndex : Children[*BoneIndex])
+				for (const int32 ChildIndex : Children[BoneIndex])
 				{
-					FGuid ChildGuid = Guids[ChildIndex];
-					FGeometryCollectionTreeItemPtr ChildPtr = NodesMap.FindRef(ChildGuid);
+					FGeometryCollectionTreeItemPtr ChildPtr = ItemsByBoneIndex.FindRef(ChildIndex);
 					if (ChildPtr.IsValid())
 					{
 						OutChildren.Add(ChildPtr);
@@ -656,89 +852,67 @@ bool FGeometryCollectionTreeItemComponent::HasChildrenForBone(const FGeometryCol
 	{
 		if (const FGeometryCollection* Collection = RestCollection->GetGeometryCollection().Get())
 		{
-			if (const int32* BoneIndex = GuidIndexMap.Find(BoneItem.GetGuid()))
+			const int32 BoneIndex = BoneItem.GetBoneIndex();
+			if (ensure(BoneIndex >= 0 && BoneIndex < Collection->NumElements(FGeometryCollection::TransformGroup)))
 			{
-				return Collection->Children[(*BoneIndex)].Num() > 0;
+				return Collection->Children[(BoneIndex)].Num() > 0;
 			}
 		}
 	}
 	return false;
 }
 
-FText FGeometryCollectionTreeItemComponent::GetDisplayNameForBone(const FGuid& Guid) const
+void FGeometryCollectionTreeItemComponent::ExpandAll()
 {
-	if (!Component.IsValid())
-	{
-		return LOCTEXT("BoneNotFound", "Bone Not Found, Invalid Geometry Collection");
-	}
-	if (const UGeometryCollection* RestCollection = Component->GetRestCollection())
-	{
-		if (FGeometryCollection* Collection = RestCollection->GetGeometryCollection().Get())
-		{
-			const TManagedArray<FString>& BoneNames = Collection->BoneName;
+	TreeView->SetItemExpansion(AsShared(), true);
 
-			if (const int32* BoneIndex = GuidIndexMap.Find(Guid))
-			{
-				if (*BoneIndex < BoneNames.Num())
-				{
-					return FText::FromString(BoneNames[*BoneIndex]);
-				}
-				else
-				{
-					return FText::Format(LOCTEXT("BoneNameNotFound", "Bone Name Not Found: Index {0}"), (*BoneIndex));
-				}
-			}
-		}
+	for (const auto& Elem : ItemsByBoneIndex)
+	{
+	    TreeView->SetItemExpansion(Elem.Value, true);
 	}
-	return LOCTEXT("BoneNotFound", "Bone Not Found, Invalid Geometry Collection");
 }
 
-void FGeometryCollectionTreeItemComponent::ExpandAll()
+void FGeometryCollectionTreeItemComponent::GenerateDataCollection()
 {
-	TreeView->SetItemExpansion(AsShared(), true);
-
-	for (auto& Elem : NodesMap)
+	DataCollection.Reset();
+	if (Component.IsValid())
 	{
-	    TreeView->SetItemExpansion(Elem.Value, true);
+		UOutlinerSettings* OutlinerSettings = GetMutableDefault<UOutlinerSettings>();
+		DataCollectionFacade.FillFromGeometryCollectionComponent(*Component, OutlinerSettings->ColumnMode);
 	}
 }
 
 void FGeometryCollectionTreeItemComponent::RegenerateChildren()
 {
+	GenerateDataCollection();
 	if(Component.IsValid() && Component->GetRestCollection())
 	{
 		//@todo Fracture:  This is potentially very expensive to refresh with giant trees
 		FGeometryCollection* Collection = Component->GetRestCollection()->GetGeometryCollection().Get();
 
-		NodesMap.Empty();
-		GuidIndexMap.Empty();
-		MyChildren.Empty();
+		ItemsByBoneIndex.Empty();
+		ChildItems.Empty();
 
 		if (Collection)
 		{
-			int32 NumElements = Collection->NumElements(FGeometryCollection::TransformGroup);
-
-			GeometryCollection::GenerateTemporaryGuids(Collection, 0, true);
-			const TManagedArray<FGuid>& Guids = Collection->GetAttribute<FGuid>("GUID", "Transform");
+			const int32 NumElements = Collection->NumElements(FGeometryCollection::TransformGroup);
 			const TManagedArray<int32>& Parents = Collection->Parent;
-			// const TManagedArray<FString>& BoneNames = CollectionPtr->BoneName;
 
 			RootIndex = FGeometryCollection::Invalid;
 
 			// Add a sub item to the outliner tree for each of the bones/chunks in this GeometryCollection
-			for (int32 Index = 0; Index < NumElements; Index++)
+			for (int32 BoneIndex = 0; BoneIndex < NumElements; BoneIndex++)
 			{
-				if (FilterBoneIndex(Index))
+				if (FilterBoneIndex(BoneIndex))
 				{
-					TSharedRef<FGeometryCollectionTreeItemBone> NewItem = MakeShared<FGeometryCollectionTreeItemBone>(Guids[Index], Index, this);
-					if (Parents[Index] == RootIndex)
+					TSharedRef<FGeometryCollectionTreeItemBone> NewItem = MakeShared<FGeometryCollectionTreeItemBone>(BoneIndex, this);
+					if (Parents[BoneIndex] == RootIndex)
 					{
 						// The actual children directly beneath this node are the ones without a parent.  The rest are children of children
-						MyChildren.Add(NewItem);
+						ChildItems.Add(NewItem);
 					}
 
-					NodesMap.Add(Guids[Index], NewItem);
-					GuidIndexMap.Add(Guids[Index], Index);
+					ItemsByBoneIndex.Add(BoneIndex, NewItem);
 				}			
 			}
 
@@ -796,125 +970,55 @@ bool FGeometryCollectionTreeItemComponent::FilterBoneIndex(int32 BoneIndex) cons
 	return true;	
 }
 
+const FGeometryCollectionItemDataFacade& FGeometryCollectionTreeItemBone::GetDataCollectionFacade() const
+{
+	ensure(ParentComponentItem);
+	return ParentComponentItem->GetDataCollectionFacade();
+}
+
 void FGeometryCollectionTreeItemBone::UpdateItemFromCollection()
 {
 	const int32 ItemBoneIndex = GetBoneIndex();
-
 	constexpr FLinearColor InvalidColor(0.1f, 0.1f, 0.1f);
-	
 	ItemColor = InvalidColor;
-	RelativeSize = 0;
-	Volume = 0;
-	Damage = 0;
-	DamageThreshold = 0;
-	Broken = false;
-	InitialState = INDEX_NONE;
-	Anchored = false;
-	RemoveOnBreakAvailable = false;
-	RemoveOnBreak = FRemoveOnBreakData::DisabledPackedData;
-	ImportedCollisionsAvailable = false;
-	ImportedCollisionsUsed = false;
-	
-	// Name / ItemText
+
+	const FGeometryCollectionItemDataFacade& DataCollectionFacade = GetDataCollectionFacade();
+
 	UOutlinerSettings* OutlinerSettings = GetMutableDefault<UOutlinerSettings>();
-	ItemText = ParentComponentItem->GetDisplayNameForBone(Guid);
-	if (OutlinerSettings->ItemText == EOutlinerItemNameEnum::BoneIndex)
+
+	const bool bHasLevelAttribute = DataCollectionFacade.IsLevelAttributeValid();
+	if (bHasLevelAttribute && OutlinerSettings->ColorByLevel)
 	{
-		ItemText = FText::FromString(FString::FromInt(ItemBoneIndex));
+		const TManagedArray<int32>& Level = DataCollectionFacade.GetLevel();
+		ItemColor = FSlateColor(FGeometryCollectionTreeItem::GetColorPerDepth((uint32)Level[ItemBoneIndex]));
 	}
-
-	// Set color according to simulation type
-	const UGeometryCollectionComponent* Component = ParentComponentItem->GetComponent();
-	if (Component)
+	else if (DataCollectionFacade.IsSimulationtypeAttributeValid())
 	{
-		const FDamageCollector* Collector = Component->GetRunTimeDataCollector();
-		const UGeometryCollection* RestCollection = Component->GetRestCollection();
-		if (RestCollection && IsValidChecked(RestCollection))
+		const TManagedArray<int32>& SimulationType = DataCollectionFacade.GetSimulationType();
+		switch (SimulationType[ItemBoneIndex])
 		{
-			TSharedPtr<FGeometryCollection, ESPMode::ThreadSafe> GeometryCollectionPtr = RestCollection->GetGeometryCollection();
-			const TManagedArray<int32>& SimulationType = GeometryCollectionPtr->SimulationType;
-			const TManagedArray<float>* RelativeSizes = GeometryCollectionPtr->FindAttribute<float>("Size", FTransformCollection::TransformGroup);
-			const TManagedArray<float>* Volumes = GeometryCollectionPtr->FindAttribute<float>("Volume", FTransformCollection::TransformGroup);
-			GeometryCollection::Facades::FCollectionRemoveOnBreakFacade RemoveOnBreakFacade(*GeometryCollectionPtr);
-
-			using FImplicitGeom = FGeometryDynamicCollection::FSharedImplicit;
-			const TManagedArray<FImplicitGeom>* ExternalCollisions = GeometryCollectionPtr->FindAttribute<FImplicitGeom>("ExternalCollisions", FTransformCollection::TransformGroup);
+		case FGeometryCollection::ESimulationTypes::FST_None:
+			ItemColor = FLinearColor::Green;
+			break;
 
-			if (ensure(ItemBoneIndex >= 0 && ItemBoneIndex < SimulationType.Num()))
+		case FGeometryCollection::ESimulationTypes::FST_Rigid:
+		{
+			bool IsVisible = true;
+			if (DataCollectionFacade.IsVisibleAttributeValid())
 			{
-				const bool bHasLevelAttribute = GeometryCollectionPtr->HasAttribute("Level", FGeometryCollection::TransformGroup);
-				if (bHasLevelAttribute && OutlinerSettings->ColorByLevel)
-				{
-					const TManagedArray<int32>& Level = GeometryCollectionPtr->GetAttribute<int32>("Level", FTransformCollection::TransformGroup);
-					ItemColor = FSlateColor(FGeometryCollectionTreeItem::GetColorPerDepth((uint32)Level[ItemBoneIndex]));
-				}
-				else
-				{
-					switch (SimulationType[ItemBoneIndex])
-					{
-					case FGeometryCollection::ESimulationTypes::FST_None:
-						ItemColor = FLinearColor::Green;
-						break;
-
-					case FGeometryCollection::ESimulationTypes::FST_Rigid:
-						if (GeometryCollectionPtr->IsVisible(ItemBoneIndex))
-						{
-							ItemColor = FSlateColor::UseForeground();
-						}
-						else
-						{
-							ItemColor = InvalidColor;
-						}
-						break;
-
-					case FGeometryCollection::ESimulationTypes::FST_Clustered:
-						ItemColor = FSlateColor(FColor::Cyan);
-						break;
-
-					default:
-						ensureMsgf(false, TEXT("Invalid Geometry Collection simulation type encountered."));
-						break;
-					}
-				}
-
-				Chaos::Facades::FCollectionAnchoringFacade AnchoringFacade(*GeometryCollectionPtr);
-				InitialState = static_cast<int32>(AnchoringFacade.GetInitialDynamicState(ItemBoneIndex));
-				if (AnchoringFacade.HasAnchoredAttribute())
-				{
-					Anchored = AnchoringFacade.IsAnchored(ItemBoneIndex);
-				}
-
-				if (RelativeSizes)
-				{
-					RelativeSize = (*RelativeSizes)[ItemBoneIndex];
-				}
-				
-				if (Volumes)
-				{
-					Volume = FMath::Pow((*Volumes)[ItemBoneIndex], 1./3);
-				}
-
-				if (Collector)
-				{
-					const FDamageCollector::FDamageData& DamageData = (*Collector)[ItemBoneIndex];
-					Damage = DamageData.MaxDamages;
-					DamageThreshold = DamageData.DamageThreshold;
-					Broken = DamageData.bIsBroken;
-				}
-				RemoveOnBreakAvailable = RemoveOnBreakFacade.IsValid();
-				if (RemoveOnBreakAvailable)
-				{
-					RemoveOnBreak = RemoveOnBreakFacade.GetData(ItemBoneIndex);
-				}
+				IsVisible = DataCollectionFacade.GetVisible()[ItemBoneIndex];
+			}
+			ItemColor = IsVisible ? FSlateColor::UseForeground() : InvalidColor;
+			break;
+		}
 
-				ImportedCollisionsUsed = RestCollection->bImportCollisionFromSource;
-				if (ExternalCollisions)
-				{
-					ImportedCollisionsAvailable = (*ExternalCollisions)[ItemBoneIndex] != nullptr;
-				}
+		case FGeometryCollection::ESimulationTypes::FST_Clustered:
+			ItemColor = FSlateColor(FColor::Cyan);
+			break;
 
-				IsCluster = (GeometryCollectionPtr->Children[ItemBoneIndex].Num() > 0);
-			}
+		default:
+			ensureMsgf(false, TEXT("Invalid Geometry Collection simulation type encountered."));
+			break;
 		}
 	}
 }
@@ -922,6 +1026,7 @@ void FGeometryCollectionTreeItemBone::UpdateItemFromCollection()
 TSharedRef<ITableRow> FGeometryCollectionTreeItemBone::MakeTreeRowWidget(const TSharedRef<STableViewBase>& InOwnerTable, bool bIsPinned)
 {
 	UpdateItemFromCollection();
+	const FGeometryCollectionItemDataFacade& DataCollectionFacade = GetDataCollectionFacade();
 	if (bIsPinned)
 	{
 		return SNew(STableRow<FGeometryCollectionTreeItemPtr>, InOwnerTable)
@@ -933,28 +1038,50 @@ TSharedRef<ITableRow> FGeometryCollectionTreeItemBone::MakeTreeRowWidget(const T
 					.AutoWidth()
 					[
 						SNew(STextBlock)
-						.Text(ItemText)
+						.Text(FText::AsNumber(BoneIndex))
 						.ColorAndOpacity(ItemColor)
 					]
+				+ SHorizontalBox::Slot()
+					.Padding(2.0f, 4.0f)
+					.AutoWidth()
+					[
+						SNew(STextBlock)
+						.Text(FText::FromString(DataCollectionFacade.GetBoneName(BoneIndex)))
+					.ColorAndOpacity(ItemColor)
+					]
 			];
 	}
 	return SNew(SGeometryCollectionOutlinerRow, InOwnerTable, SharedThis(this));
 }
 
-TSharedRef<SWidget> FGeometryCollectionTreeItemBone::MakeNameColumnWidget() const
+TSharedRef<SWidget> FGeometryCollectionTreeItemBone::MakeBoneIndexColumnWidget() const
 {
 	return SNew(SHorizontalBox)
-			+ SHorizontalBox::Slot()
-				.AutoWidth()
-				[
-					SNew(STextBlock)
-					.Text(ItemText)
-					.ColorAndOpacity(ItemColor)
-				];
+		+ SHorizontalBox::Slot()
+		.AutoWidth()
+		[
+			SNew(STextBlock)
+			.Text(FText::AsNumber(BoneIndex))
+		.ColorAndOpacity(ItemColor)
+		];
+}
+
+TSharedRef<SWidget> FGeometryCollectionTreeItemBone::MakeBoneNameColumnWidget() const
+{
+	const FGeometryCollectionItemDataFacade& DataCollectionFacade = GetDataCollectionFacade();
+	return SNew(SHorizontalBox)
+		+ SHorizontalBox::Slot()
+			.AutoWidth()
+			[
+				SNew(STextBlock)
+				.Text(FText::FromString(DataCollectionFacade.GetBoneName(BoneIndex)))
+				.ColorAndOpacity(ItemColor)
+			];
 }
 
 TSharedRef<SWidget> FGeometryCollectionTreeItemBone::MakeRelativeSizeColumnWidget() const
 {
+	const FGeometryCollectionItemDataFacade& DataCollectionFacade = GetDataCollectionFacade();
 	static const FNumberFormattingOptions FormatOptions = FNumberFormattingOptions()
 				.SetMinimumFractionalDigits(2)
 				.SetMaximumFractionalDigits(2);
@@ -965,13 +1092,14 @@ TSharedRef<SWidget> FGeometryCollectionTreeItemBone::MakeRelativeSizeColumnWidge
 			.HAlign(HAlign_Right)
 			[
 				SNew(STextBlock)
-				.Text(FText::AsNumber(RelativeSize, &FormatOptions))
+				.Text(FText::AsNumber(DataCollectionFacade.GetRelativeSize(BoneIndex), &FormatOptions))
 				.ColorAndOpacity(ItemColor)
 			];
 }
 
 TSharedRef<SWidget> FGeometryCollectionTreeItemBone::MakeVolumeColumnWidget() const
 {
+	const FGeometryCollectionItemDataFacade& DataCollectionFacade = GetDataCollectionFacade();
 	static const FNumberFormattingOptions FormatOptions = FNumberFormattingOptions()
 				.SetMinimumFractionalDigits(2)
 				.SetMaximumFractionalDigits(2);
@@ -982,13 +1110,14 @@ TSharedRef<SWidget> FGeometryCollectionTreeItemBone::MakeVolumeColumnWidget() co
 			.HAlign(HAlign_Right)
 			[
 				SNew(STextBlock)
-				.Text(FText::AsNumber(Volume, &FormatOptions))
+				.Text(FText::AsNumber(DataCollectionFacade.GetVolume(BoneIndex), &FormatOptions))
 				.ColorAndOpacity(ItemColor)
 			];
 }
 
 TSharedRef<SWidget> FGeometryCollectionTreeItemBone::MakeDamagesColumnWidget() const
 {
+	const FGeometryCollectionItemDataFacade& DataCollectionFacade = GetDataCollectionFacade();
 	static const FNumberFormattingOptions FormatOptions = FNumberFormattingOptions()
 			.SetMinimumFractionalDigits(1)
 			.SetMaximumFractionalDigits(1);
@@ -999,44 +1128,76 @@ TSharedRef<SWidget> FGeometryCollectionTreeItemBone::MakeDamagesColumnWidget() c
 			.HAlign(HAlign_Right)
 			[
 				SNew(STextBlock)
-				.Text(FText::AsNumber(Damage, &FormatOptions))
+				.Text(FText::AsNumber(DataCollectionFacade.GetDamage(BoneIndex), &FormatOptions))
 				.ColorAndOpacity(ItemColor)
 			];
 }
 
 TSharedRef<SWidget> FGeometryCollectionTreeItemBone::MakeDamageThresholdColumnWidget() const
 {
+	const FGeometryCollectionItemDataFacade& DataCollectionFacade = GetDataCollectionFacade();
+	static const FNumberFormattingOptions FormatOptions = FNumberFormattingOptions()
+		.SetMinimumFractionalDigits(1)
+		.SetMaximumFractionalDigits(1);
+
 	return SNew(SHorizontalBox)
 		+ SHorizontalBox::Slot()
 			.Padding(12.f, 0.f)
 			.HAlign(HAlign_Right)
 			[
 				SNew(STextBlock)
-				.Text(FText::AsNumber(DamageThreshold))
+				.Text(FText::AsNumber(DataCollectionFacade.GetDamageThreshold(BoneIndex), &FormatOptions))
 				.ColorAndOpacity(ItemColor)
 			];
 }
 
 TSharedRef<SWidget> FGeometryCollectionTreeItemBone::MakeBrokenColumnWidget() const
 {
+	const FGeometryCollectionItemDataFacade& DataCollectionFacade = GetDataCollectionFacade();
 	return SNew(SHorizontalBox)
 		+ SHorizontalBox::Slot()
 			.Padding(12.f, 0.f)
 			.HAlign(HAlign_Center)
 			[
 				SNew(STextBlock)
-				.Text(FText::FromString(Broken? TEXT("Broken"): TEXT("")))
+				.Text(FText::FromString(DataCollectionFacade.IsBroken(BoneIndex) ? TEXT("Broken") : TEXT("")))
 				.ColorAndOpacity(ItemColor)
 			];
 }
 
-static FText FormatRemoveOnBreakTimeData(bool bDataAvailable, bool bEnabled, const FVector2f& Timer)
+static FText FormatRemoveOnBreak_BreakTimer(const FGeometryCollectionItemDataFacade& DataCollectionFacade, int32 BoneIndex)
 {
-	if (bDataAvailable)
+	if (DataCollectionFacade.IsRemoveOnBreakAttributeValid())
 	{
-		if (bEnabled)
+		FRemoveOnBreakData Data = DataCollectionFacade.GetRemoveOnBreakData(BoneIndex);
+		if (Data.IsEnabled())
 		{
-			return FText::Format(LOCTEXT("GCOutliner_RemoveOnBreakTimer_Format", "[{0}s, {1}s]"),Timer.X, Timer.Y);
+			const FVector2f BreakTimer = Data.GetBreakTimer();
+			return FText::Format(LOCTEXT("GCOutliner_RemoveOnBreakTimer_Format", "[{0}s, {1}s]"), BreakTimer.X, BreakTimer.Y);
+		}
+		return FText(LOCTEXT("GCOutliner_RemoveOnBreakTimer_Empty", "[-, -]"));
+	}
+	return FText();
+}
+
+static FText FormatRemoveOnBreak_RemovalTimer(const FGeometryCollectionItemDataFacade& DataCollectionFacade, int32 BoneIndex)
+{
+	if (DataCollectionFacade.IsRemoveOnBreakAttributeValid())
+	{
+		FRemoveOnBreakData Data = DataCollectionFacade.GetRemoveOnBreakData(BoneIndex);
+		if (Data.IsEnabled())
+		{
+			const TManagedArray<int32>& SimulationType = DataCollectionFacade.GetSimulationType();
+			const bool IsCluster = (SimulationType.Num() > 0) ? (SimulationType[BoneIndex] == FGeometryCollection::ESimulationTypes::FST_Clustered) : false;
+
+			const bool EnableRemovalTimer = !(IsCluster && Data.GetClusterCrumbling());
+			const bool ClusterCrumbling = (IsCluster && Data.GetClusterCrumbling());
+			if (ClusterCrumbling)
+			{
+				return FText(LOCTEXT("GCOutliner_ClusterCrumbling_Text", "Cluster Crumbling"));
+			}
+			const FVector2f RemovalTimer = Data.GetRemovalTimer();
+			return FText::Format(LOCTEXT("GCOutliner_RemoveOnBreakTimer_Format", "[{0}s, {1}s]"), RemovalTimer.X, RemovalTimer.Y);
 		}
 		return FText(LOCTEXT("GCOutliner_RemoveOnBreakTimer_Empty", "[-, -]"));
 	}
@@ -1045,47 +1206,39 @@ static FText FormatRemoveOnBreakTimeData(bool bDataAvailable, bool bEnabled, con
 
 TSharedRef<SWidget> FGeometryCollectionTreeItemBone::MakePostBreakTimeColumnWidget() const
 {
+	const FGeometryCollectionItemDataFacade& DataCollectionFacade = GetDataCollectionFacade();
 	return SNew(SHorizontalBox)
 		+ SHorizontalBox::Slot()
 			.Padding(12.f, 0.f)
 			.HAlign(HAlign_Center)
 			[
 				SNew(STextBlock)
-				.Text(FormatRemoveOnBreakTimeData(RemoveOnBreakAvailable, RemoveOnBreak.IsEnabled(), RemoveOnBreak.GetBreakTimer()))
+				.Text(FormatRemoveOnBreak_BreakTimer(DataCollectionFacade, BoneIndex))
 				.ColorAndOpacity(ItemColor)
 			];
 }
 
 TSharedRef<SWidget> FGeometryCollectionTreeItemBone::MakeRemovalTimeColumnWidget() const
 {
-	const bool EnableRemovalTimer = RemoveOnBreak.IsEnabled() && !(IsCluster && RemoveOnBreak.GetClusterCrumbling());
-	const bool ClusterCrumbling = RemoveOnBreak.IsEnabled() && (IsCluster && RemoveOnBreak.GetClusterCrumbling());
-	FText RemovalTimeText; 
-	if (ClusterCrumbling)
-	{
-		RemovalTimeText = FText(LOCTEXT("GCOutliner_ClusterCrumbling_Text", "Cluster Crumbling"));
-	}
-	else
-	{
-		RemovalTimeText = FormatRemoveOnBreakTimeData(RemoveOnBreakAvailable, EnableRemovalTimer,  RemoveOnBreak.GetRemovalTimer());
-	}
+	const FGeometryCollectionItemDataFacade& DataCollectionFacade = GetDataCollectionFacade();
 	return SNew(SHorizontalBox)
 	+ SHorizontalBox::Slot()
 		.Padding(12.f, 0.f)
 		.HAlign(HAlign_Center)
 		[
 			SNew(STextBlock)
-			.Text(RemovalTimeText)
+			.Text(FormatRemoveOnBreak_RemovalTimer(DataCollectionFacade, BoneIndex))
 			.ColorAndOpacity(ItemColor)
 		];
 }
 
 TSharedRef<SWidget> FGeometryCollectionTreeItemBone::MakeImportedCollisionsColumnWidget() const
 {
+	const FGeometryCollectionItemDataFacade& DataCollectionFacade = GetDataCollectionFacade();
 	FText ImportedCollisionText = LOCTEXT("GCOutliner_ImportedCollision_NotAvailable", "-");
-	if (ImportedCollisionsAvailable)
+	if (DataCollectionFacade.HasSourceCollision(BoneIndex))
 	{
-		if (ImportedCollisionsUsed)
+		if (DataCollectionFacade.IsSourceCollisionUsed(BoneIndex))
 		{
 			ImportedCollisionText = LOCTEXT("GCOutliner_ImportedCollision_Used", "Used");
 		}
@@ -1107,24 +1260,26 @@ TSharedRef<SWidget> FGeometryCollectionTreeItemBone::MakeImportedCollisionsColum
 
 TSharedRef<SWidget> FGeometryCollectionTreeItemBone::MakeInitialStateColumnWidget() const
 {
+	const FGeometryCollectionItemDataFacade& DataCollectionFacade = GetDataCollectionFacade();
 	return SNew(SHorizontalBox)
 		+ SHorizontalBox::Slot()
 			.Padding(12.f, 0.f)
 			[
 				SNew(STextBlock)
-				.Text(GetTextFromInitialDynamicState(InitialState))
+				.Text(GetTextFromInitialDynamicState(DataCollectionFacade.GetInitialState(BoneIndex)))
 				.ColorAndOpacity(ItemColor)
 			];
 }
 
 TSharedRef<SWidget> FGeometryCollectionTreeItemBone::MakeAnchoredColumnWidget() const
 {
+	const FGeometryCollectionItemDataFacade& DataCollectionFacade = GetDataCollectionFacade();
 	return SNew(SHorizontalBox)
 		+ SHorizontalBox::Slot()
 			.Padding(12.f, 0.f)
 	[
 				SNew(STextBlock)
-				.Text(GetTextFromAnchored(Anchored))
+				.Text(GetTextFromAnchored(DataCollectionFacade.IsAnchored(BoneIndex)))
 				.ColorAndOpacity(ItemColor)
 			];
 }
diff --git a/Engine/Plugins/Experimental/ChaosEditor/Source/FractureEditor/Private/SGeometryCollectionOutliner.h b/Engine/Plugins/Experimental/ChaosEditor/Source/FractureEditor/Private/SGeometryCollectionOutliner.h
--- a/Engine/Plugins/Experimental/ChaosEditor/Source/FractureEditor/Private/SGeometryCollectionOutliner.h
+++ b/Engine/Plugins/Experimental/ChaosEditor/Source/FractureEditor/Private/SGeometryCollectionOutliner.h
@@ -40,10 +40,11 @@ enum class EOutlinerItemNameEnum : uint8
 UENUM(BlueprintType)
 enum class EOutlinerColumnMode : uint8
 {
-	StateAndSize = 0		UMETA(DisplayName = "State And Size"),
+	State = 0				UMETA(DisplayName = "State"),
 	Damage = 1				UMETA(DisplayName = "Damage"),
 	Removal = 2				UMETA(DisplayName = "Removal"),
 	Collision = 3			UMETA(DisplayName = "Collision"),
+	Size = 4				UMETA(DisplayName = "Size"),
 };
 
 /** Settings for Outliner configuration. **/
@@ -55,10 +56,6 @@ class UOutlinerSettings : public UObject
 public:
 	UOutlinerSettings(const FObjectInitializer& ObjInit);
 
-	/** What is displayed in Outliner text */
-	UPROPERTY(EditAnywhere, Category = OutlinerSettings, meta = (DisplayName = "Item Text"))
-	EOutlinerItemNameEnum ItemText;
-
 	/** whether to use level coloring */
 	UPROPERTY(EditAnywhere, Category = OutlinerSettings, meta = (DisplayName = "Color By Level"))
 	bool ColorByLevel;
@@ -91,12 +88,62 @@ public:
 	static FColor GetColorPerDepth(uint32 Depth);
 };
 
+using FRemoveOnBreakData = GeometryCollection::Facades::FRemoveOnBreakData;
+
+class FGeometryCollectionItemDataFacade
+{
+public:
+	FGeometryCollectionItemDataFacade(FManagedArrayCollection& InCollection);
+
+	void FillFromGeometryCollectionComponent(const UGeometryCollectionComponent& GeometryCollectionComponent, EOutlinerColumnMode ColumnMode);
+
+	bool IsLevelAttributeValid() const { return LevelAttribute.IsValid(); }
+	bool IsSimulationtypeAttributeValid() const { return SimulationTypeAttribute.IsValid(); }
+	bool IsVisibleAttributeValid() const { return VisibleAttribute.IsValid(); }
+	bool IsRemoveOnBreakAttributeValid() const { return RemoveOnBreakAttribute.IsValid(); }
+
+	const TManagedArray<int32>& GetLevel() const { return LevelAttribute.Get(); }
+	const TManagedArray<int32>& GetSimulationType() const { return SimulationTypeAttribute.Get(); }
+	const TManagedArray<bool>& GetVisible() const { return VisibleAttribute.Get(); }
+
+	FString GetBoneName(int32 Index) const;
+	float GetRelativeSize(int32 Index) const;
+	float GetVolume(int32 Index) const;
+	int32 GetInitialState(int32 Index) const;
+	bool IsAnchored(int32 Index) const;
+	float GetDamage(int32 Index) const;
+	float GetDamageThreshold(int32 Index) const;
+	bool IsBroken(int32 Index) const;
+	FRemoveOnBreakData GetRemoveOnBreakData(int32 Index) const;
+	bool HasSourceCollision(int32 Index) const;
+	bool IsSourceCollisionUsed(int32 Index) const;
+
+private:
+	FManagedArrayCollection&			DataCollection;
+	TManagedArrayAccessor<FString>		BoneNameAttribute;
+	TManagedArrayAccessor<int32>		LevelAttribute;
+	//TManagedArrayAccessor<TArray<int32>>ChildrenAttribute;
+	TManagedArrayAccessor<bool>			VisibleAttribute;
+	TManagedArrayAccessor<int32>		InitialStateAttribute;
+	TManagedArrayAccessor<float>		RelativeSizeAttribute;
+	TManagedArrayAccessor<float>		VolumeAttribute;
+	TManagedArrayAccessor<bool>			AnchoredAttribute;
+	TManagedArrayAccessor<float>		DamageAttribute;
+	TManagedArrayAccessor<float>		DamageThresholdAttribute;
+	TManagedArrayAccessor<bool>			BrokenStateAttribute;
+	TManagedArrayAccessor<FVector4f>	RemoveOnBreakAttribute;
+	TManagedArrayAccessor<int32>		SimulationTypeAttribute;
+	TManagedArrayAccessor<bool>			HasSourceCollisionAttribute;
+	TManagedArrayAccessor<bool>			SourceCollisionUsedAttribute;
+};
+
 class FGeometryCollectionTreeItemComponent : public FGeometryCollectionTreeItem
 {
 public:
 	FGeometryCollectionTreeItemComponent(UGeometryCollectionComponent* InComponent, TSharedPtr<STreeView<FGeometryCollectionTreeItemPtr>> InTreeView)
 		: Component(InComponent)
 		, TreeView(InTreeView)
+		, DataCollectionFacade(DataCollection)
 	{
 		RegenerateChildren();
 	}
@@ -111,14 +158,16 @@ public:
 	 
 	void GetChildrenForBone(FGeometryCollectionTreeItemBone& BoneItem, FGeometryCollectionTreeItemList& OutChildren) const;
 	bool HasChildrenForBone(const FGeometryCollectionTreeItemBone& BoneItem) const;
-	FText GetDisplayNameForBone(const FGuid& Guid) const;
 
 	void ExpandAll();
+	void GenerateDataCollection();
 	void RegenerateChildren();
 	void RequestTreeRefresh();
 
 	void SetHistogramSelection(TArray<int32>& SelectedBones);
 
+	FGeometryCollectionItemDataFacade& GetDataCollectionFacade() { return DataCollectionFacade; }
+
 private:
 	bool FilterBoneIndex(int32 BoneIndex) const;
 
@@ -128,43 +177,32 @@ private:
 	TSharedPtr<STreeView<FGeometryCollectionTreeItemPtr>> TreeView;
 
 	/** The direct children under this component */
-	TArray<FGeometryCollectionTreeItemPtr> MyChildren;
-
-	TMap<FGuid, FGeometryCollectionTreeItemPtr> NodesMap;
-	TMap<FGuid, int32> GuidIndexMap;
-	FGuid RootGuid;
+	TArray<FGeometryCollectionTreeItemPtr> ChildItems;
+	TMap<int32, FGeometryCollectionTreeItemPtr> ItemsByBoneIndex;
 	int32 RootIndex;
 
 	TArray<int32> HistogramSelection;
-};
 
-using FRemoveOnBreakData = GeometryCollection::Facades::FRemoveOnBreakData;
+	// collection used to store the displayed information
+	FManagedArrayCollection DataCollection;
+	FGeometryCollectionItemDataFacade DataCollectionFacade;
+
+};
 
 class FGeometryCollectionTreeItemBone : public FGeometryCollectionTreeItem
 {
 
 public:
-	FGeometryCollectionTreeItemBone(const FGuid NewGuid, const int32 InBoneIndex, FGeometryCollectionTreeItemComponent* InParentComponentItem)
-		: Guid(NewGuid)
-		, BoneIndex(InBoneIndex)
+	FGeometryCollectionTreeItemBone(const int32 InBoneIndex, FGeometryCollectionTreeItemComponent* InParentComponentItem)
+		: BoneIndex(InBoneIndex)
 		, ParentComponentItem(InParentComponentItem)
 		, ItemColor(FSlateColor::UseForeground())
-		, RelativeSize(0)
-		, Volume(0)
-		, InitialState(INDEX_NONE)
-		, Anchored(false)
-		, Damage(0)
-		, DamageThreshold(0)
-		, Broken(false)
-		, RemoveOnBreakAvailable(false)
-		, RemoveOnBreak(FRemoveOnBreakData::DisabledPackedData)
-		, ImportedCollisionsAvailable(false)
-		, ImportedCollisionsUsed(false)
 	{}
 
 	/** FGeometryCollectionTreeItem interface */
 	TSharedRef<ITableRow> MakeTreeRowWidget(const TSharedRef<STableViewBase>& InOwnerTable, bool bIsPinned = false);
-	TSharedRef<SWidget> MakeNameColumnWidget() const;
+	TSharedRef<SWidget> MakeBoneIndexColumnWidget() const;
+	TSharedRef<SWidget> MakeBoneNameColumnWidget() const;
 	TSharedRef<SWidget> MakeRelativeSizeColumnWidget() const;
 	TSharedRef<SWidget> MakeVolumeColumnWidget() const;
 	TSharedRef<SWidget> MakeDamagesColumnWidget() const;
@@ -179,7 +217,6 @@ public:
 	virtual void GetChildren(FGeometryCollectionTreeItemList& OutChildren) override;
 	virtual int32 GetBoneIndex() const override { return BoneIndex; }
 	virtual UGeometryCollectionComponent* GetComponent() const { return ParentComponentItem->GetComponent(); }
-	const FGuid& GetGuid() const { return Guid; }
 	bool HasChildren() const;
 	
 protected:
@@ -189,33 +226,21 @@ protected:
 
 private:
 	void UpdateItemFromCollection();
+
+	const FGeometryCollectionItemDataFacade& GetDataCollectionFacade() const;
 	
 private:
-	const FGuid Guid;
 	const int32 BoneIndex;
 	FGeometryCollectionTreeItemComponent* ParentComponentItem;
-	
 	FSlateColor ItemColor;
-	FText ItemText;
-	float RelativeSize;
-	float Volume;
-	int32 InitialState;
-	bool Anchored;
-	float Damage;
-	float DamageThreshold;
-	bool Broken;
-	bool RemoveOnBreakAvailable;
-	bool IsCluster;
-	FRemoveOnBreakData RemoveOnBreak;
-	bool ImportedCollisionsAvailable;
-	bool ImportedCollisionsUsed;
 };
 
 typedef TSharedPtr<class FGeometryCollectionTreeItemBone> FGeometryCollectionTreeItemBonePtr;
 
 namespace SGeometryCollectionOutlinerColumnID
 {
-	const FName Bone("Bone");
+	const FName BoneIndex("Index");
+	const FName BoneName("Name");
 	// State and Size column mode
 	const FName RelativeSize("Relative Size");
 	const FName Volume("Volume");
@@ -272,6 +297,7 @@ public:
 
 	void RegenerateItems();
 	void RegenerateHeader();
+	void RegenerateRootData();
 
 	TSharedRef<ITableRow> MakeTreeRowWidget(FGeometryCollectionTreeItemPtr InInfo, const TSharedRef<STableViewBase>& OwnerTable);
 	TSharedRef<ITableRow> OnGeneratePinnedRowWidget(FGeometryCollectionTreeItemPtr InItem, const TSharedRef<STableViewBase>& InOwnerTable, bool bPinned);
