Fracture editor : fix geometry collection statistic window not displaying properly when GC is inside a blueprint

diff --git a/Engine/Plugins/Experimental/ChaosEditor/Source/FractureEditor/Private/FractureEditorModeToolkit.cpp b/Engine/Plugins/Experimental/ChaosEditor/Source/FractureEditor/Private/FractureEditorModeToolkit.cpp
--- a/Engine/Plugins/Experimental/ChaosEditor/Source/FractureEditor/Private/FractureEditorModeToolkit.cpp
+++ b/Engine/Plugins/Experimental/ChaosEditor/Source/FractureEditor/Private/FractureEditorModeToolkit.cpp
@@ -1885,17 +1885,17 @@ void FFractureEditorModeToolkit::GetStatisticsSummary(FGeometryCollectionStatist
 	{
 		for (FSelectionIterator Iter(*SelectedActors); Iter; ++Iter)
 		{
-			if (AGeometryCollectionActor* Actor = Cast<AGeometryCollectionActor>(*Iter))
+			if (AActor* Actor = Cast<AActor>(*Iter))
 			{
-				if (UGeometryCollectionComponent* Component = Actor->GetGeometryCollectionComponent())
+				TInlineComponentArray<UGeometryCollectionComponent*> GeometryCollectionComponents;
+				Actor->GetComponents(GeometryCollectionComponents);
+				for (UGeometryCollectionComponent* GeometryCollectionComponent : GeometryCollectionComponents)
 				{
-					if (const UGeometryCollection* RestCollection = Component->GetRestCollection())
+					if (const UGeometryCollection* RestCollection = GeometryCollectionComponent->GetRestCollection())
 					{
-						const FGeometryCollection* GeometryCollection = RestCollection->GetGeometryCollection().Get();
-
-						if (GeometryCollection != nullptr)
+						if (RestCollection->GetGeometryCollection())
 						{
-							GeometryCollectionArray.Add(GeometryCollection);
+							GeometryCollectionArray.Add(RestCollection->GetGeometryCollection().Get());
 						}
 					}
 				}
@@ -1903,7 +1903,6 @@ void FFractureEditorModeToolkit::GetStatisticsSummary(FGeometryCollectionStatist
 		}
 	}
 
-
 	if (GeometryCollectionArray.Num() > 0)
 	{
 		TArray<int32> LevelTransformsAll;
