Fracture editor : fix size attribute not properly representing the volumetric unit in the outliner

diff --git a/Engine/Plugins/Experimental/ChaosEditor/Source/FractureEditor/Private/SGeometryCollectionOutliner.cpp b/Engine/Plugins/Experimental/ChaosEditor/Source/FractureEditor/Private/SGeometryCollectionOutliner.cpp
--- a/Engine/Plugins/Experimental/ChaosEditor/Source/FractureEditor/Private/SGeometryCollectionOutliner.cpp
+++ b/Engine/Plugins/Experimental/ChaosEditor/Source/FractureEditor/Private/SGeometryCollectionOutliner.cpp
@@ -614,7 +614,7 @@ FGeometryCollectionItemDataFacade::FGeometryCollectionItemDataFacade(FManagedArr
 	, VisibleAttribute(InCollection, "Visible", DataCollectionGroup)
 	, InitialStateAttribute(InCollection, "InitialState", DataCollectionGroup)
 	, RelativeSizeAttribute(InCollection, "RelativeSize", DataCollectionGroup)
-	, VolumeAttribute(InCollection, "Volume", DataCollectionGroup)
+	, VolumetricUnitAttribute(InCollection, "VolumetricUnit", DataCollectionGroup)
 	, AnchoredAttribute(InCollection, "Anchored", DataCollectionGroup)
 	, DamageAttribute(InCollection, "Damage", DataCollectionGroup)
 	, DamageThresholdAttribute(InCollection, "DamageThreshold", DataCollectionGroup)
@@ -642,9 +642,9 @@ float FGeometryCollectionItemDataFacade::GetRelativeSize(int32 Index) const
 	return GetAttributeValue(RelativeSizeAttribute, Index, 0.0f);
 }
 
-float FGeometryCollectionItemDataFacade::GetVolume(int32 Index) const
+float FGeometryCollectionItemDataFacade::GetVolumetricUnit(int32 Index) const
 {
-	return GetAttributeValue(VolumeAttribute, Index, 0.0f);
+	return GetAttributeValue(VolumetricUnitAttribute, Index, 0.0f);
 }
 int32 FGeometryCollectionItemDataFacade::GetInitialState(int32 Index) const
 {
@@ -744,7 +744,13 @@ void FGeometryCollectionItemDataFacade::FillFromGeometryCollectionComponent(cons
 			const TManagedArrayAccessor<float> GCVolumeAttribute(GeometryCollection, "Volume", FGeometryCollection::TransformGroup);
 			if (GCVolumeAttribute.IsValid())
 			{
-				VolumeAttribute.Copy(GCVolumeAttribute);
+				VolumetricUnitAttribute.Copy(GCVolumeAttribute);
+				// volumetric unit is the cubic root of the volume so it is more meaningful to read and comprehend than volume
+				TManagedArray<float>& VolumetricUnit = VolumetricUnitAttribute.Modify();
+				for (int32 Index = 0; Index < VolumetricUnit.Num(); Index++)
+				{
+					VolumetricUnit[Index] = FMath::Pow(VolumetricUnit[Index], 1.0/3.0);
+				}
 			}
 		}
 		else if (ColumnMode == EOutlinerColumnMode::Damage)
@@ -1062,7 +1068,7 @@ TSharedRef<SWidget> FGeometryCollectionTreeItemBone::MakeBoneIndexColumnWidget()
 		[
 			SNew(STextBlock)
 			.Text(FText::AsNumber(BoneIndex))
-		.ColorAndOpacity(ItemColor)
+			.ColorAndOpacity(ItemColor)
 		];
 }
 
@@ -1076,6 +1082,7 @@ TSharedRef<SWidget> FGeometryCollectionTreeItemBone::MakeBoneNameColumnWidget()
 				SNew(STextBlock)
 				.Text(FText::FromString(DataCollectionFacade.GetBoneName(BoneIndex)))
 				.ColorAndOpacity(ItemColor)
+				.OverflowPolicy(ETextOverflowPolicy::Ellipsis)
 			];
 }
 
@@ -1110,7 +1117,7 @@ TSharedRef<SWidget> FGeometryCollectionTreeItemBone::MakeVolumeColumnWidget() co
 			.HAlign(HAlign_Right)
 			[
 				SNew(STextBlock)
-				.Text(FText::AsNumber(DataCollectionFacade.GetVolume(BoneIndex), &FormatOptions))
+				.Text(FText::AsNumber(DataCollectionFacade.GetVolumetricUnit(BoneIndex), &FormatOptions))
 				.ColorAndOpacity(ItemColor)
 			];
 }
diff --git a/Engine/Plugins/Experimental/ChaosEditor/Source/FractureEditor/Private/SGeometryCollectionOutliner.h b/Engine/Plugins/Experimental/ChaosEditor/Source/FractureEditor/Private/SGeometryCollectionOutliner.h
--- a/Engine/Plugins/Experimental/ChaosEditor/Source/FractureEditor/Private/SGeometryCollectionOutliner.h
+++ b/Engine/Plugins/Experimental/ChaosEditor/Source/FractureEditor/Private/SGeometryCollectionOutliner.h
@@ -108,7 +108,7 @@ public:
 
 	FString GetBoneName(int32 Index) const;
 	float GetRelativeSize(int32 Index) const;
-	float GetVolume(int32 Index) const;
+	float GetVolumetricUnit(int32 Index) const;
 	int32 GetInitialState(int32 Index) const;
 	bool IsAnchored(int32 Index) const;
 	float GetDamage(int32 Index) const;
@@ -126,7 +126,7 @@ private:
 	TManagedArrayAccessor<bool>			VisibleAttribute;
 	TManagedArrayAccessor<int32>		InitialStateAttribute;
 	TManagedArrayAccessor<float>		RelativeSizeAttribute;
-	TManagedArrayAccessor<float>		VolumeAttribute;
+	TManagedArrayAccessor<float>		VolumetricUnitAttribute;
 	TManagedArrayAccessor<bool>			AnchoredAttribute;
 	TManagedArrayAccessor<float>		DamageAttribute;
 	TManagedArrayAccessor<float>		DamageThresholdAttribute;
