Add icon for fracture mode proximity tool

diff --git a/Engine/Plugins/Experimental/ChaosEditor/Source/FractureEditor/Private/FractureEditorStyle.cpp b/Engine/Plugins/Experimental/ChaosEditor/Source/FractureEditor/Private/FractureEditorStyle.cpp
--- a/Engine/Plugins/Experimental/ChaosEditor/Source/FractureEditor/Private/FractureEditorStyle.cpp
+++ b/Engine/Plugins/Experimental/ChaosEditor/Source/FractureEditor/Private/FractureEditorStyle.cpp
@@ -73,6 +73,7 @@ FFractureEditorStyle::FFractureEditorStyle()
 	Set("FractureEditor.CustomVoronoi",        new IMAGE_BRUSH_SVG("FractureCustom", IconSize));
 	Set("FractureEditor.FixTinyGeo",	       new IMAGE_BRUSH_SVG("FractureGeoMerge", IconSize));
 	Set("FractureEditor.SetRemoveOnBreak",     new IMAGE_BRUSH_SVG("FractureSetRemoval", IconSize));
+	Set("FractureEditor.Proximity",            new IMAGE_BRUSH_SVG("FractureProximity", IconSize));
 
 	// View Settings
 	Set("FractureEditor.Exploded",             new IMAGE_BRUSH_SVG("FractureMiniExploded", LabelIconSize));
