Add background compute of voronoi diagrams
Add background compute of the noise-deformed cutting meshes
+ add previews using background computes

Move the noise preview controls (previously just in the plane fracture tool) over to common fracture settings

diff --git a/Engine/Plugins/Experimental/ChaosEditor/Source/FractureEditor/FractureEditor.Build.cs b/Engine/Plugins/Experimental/ChaosEditor/Source/FractureEditor/FractureEditor.Build.cs
--- a/Engine/Plugins/Experimental/ChaosEditor/Source/FractureEditor/FractureEditor.Build.cs
+++ b/Engine/Plugins/Experimental/ChaosEditor/Source/FractureEditor/FractureEditor.Build.cs
@@ -44,7 +44,8 @@ public class FractureEditor : ModuleRules
 				"Chaos",
 				"ToolWidgets",
 				"DeveloperSettings",
-				"FractureEngine"
+				"FractureEngine",
+				"GeometryFramework"
 
 				// ... add private dependencies that you statically link with here ...	
 			}
diff --git a/Engine/Plugins/Experimental/ChaosEditor/Source/FractureEditor/Private/FractureToolBackgroundTask.h b/Engine/Plugins/Experimental/ChaosEditor/Source/FractureEditor/Private/FractureToolBackgroundTask.h
--- a/Engine/Plugins/Experimental/ChaosEditor/Source/FractureEditor/Private/FractureToolBackgroundTask.h
+++ b/Engine/Plugins/Experimental/ChaosEditor/Source/FractureEditor/Private/FractureToolBackgroundTask.h
@@ -70,6 +70,68 @@ public:
 	}
 };
 
+template<class OpType>
+using TBackgroundOpExecuter = UE::Geometry::FAsyncTaskExecuterWithProgressCancel<UE::Geometry::TModelingOpTask<OpType>>;
+
+
+// Start a background task
+// Note: If automatically starting tasks in response to user actions, consider using a timeout to prevent launching too many tasks at once
+// (TODO: consider encapsulating the timeout logic in a small struct here)
+template<class OpType>
+TUniquePtr<TBackgroundOpExecuter<OpType>> StartBackgroundTask(TUniquePtr<OpType>&& NewOp)
+{
+	TUniquePtr<TBackgroundOpExecuter<OpType>> BackgroundTask = MakeUnique<TBackgroundOpExecuter<OpType>>(MoveTemp(NewOp));
+	BackgroundTask->StartBackgroundTask();
+	return BackgroundTask;
+}
+
+// Tick a background task started by StartBackgroundTask (above)
+// @return the result of SuccessCallback if the task has completed, false otherwise
+template<class OpType>
+bool TickBackgroundTask(TUniquePtr<TBackgroundOpExecuter<OpType>>& BackgroundTask, bool bCancel, TFunctionRef<bool(TUniquePtr<typename OpType::ResultTypeName>&&)> SuccessCallback)
+{
+	bool bSuccess = false;
+	if (bCancel)
+	{
+		BackgroundTask.Release()->CancelAndDelete();
+	}
+	else if (BackgroundTask->IsDone())
+	{
+		bSuccess = !BackgroundTask->GetTask().IsAborted();
+	}
+
+	if (bSuccess)
+	{
+		check(BackgroundTask != nullptr && BackgroundTask->IsDone());
+		TUniquePtr<OpType> Op = BackgroundTask->GetTask().ExtractOperator();
+		if (!ensure(Op))
+		{
+			return false;
+		}
+
+		TUniquePtr<typename OpType::ResultTypeName> Result = Op->ExtractResult();
+		if (!Result.IsValid())
+		{
+			return false;
+		}
+		bool bResult = SuccessCallback(MoveTemp(Result));
+		BackgroundTask.Release();
+		return bResult;
+	}
+
+	return false;
+}
+
+// Cancel the background task (or ignore if the pointer is empty)
+template<class OpType>
+void CancelBackgroundTask(TUniquePtr<TBackgroundOpExecuter<OpType>>& BackgroundTask)
+{
+	if (BackgroundTask)
+	{
+		BackgroundTask.Release()->CancelAndDelete();
+	}
+}
+
 
 // Run a blocking geometry collection op, but with a responsive cancel option
 template<class GeometryCollectionOpType, typename ResultType>
diff --git a/Engine/Plugins/Experimental/ChaosEditor/Source/FractureEditor/Private/FractureToolBrick.cpp b/Engine/Plugins/Experimental/ChaosEditor/Source/FractureEditor/Private/FractureToolBrick.cpp
--- a/Engine/Plugins/Experimental/ChaosEditor/Source/FractureEditor/Private/FractureToolBrick.cpp
+++ b/Engine/Plugins/Experimental/ChaosEditor/Source/FractureEditor/Private/FractureToolBrick.cpp
@@ -26,6 +26,7 @@ UFractureToolBrick::UFractureToolBrick(const FObjectInitializer& ObjInit)
 {
 	BrickSettings = NewObject<UFractureBrickSettings>(GetTransientPackage(), UFractureBrickSettings::StaticClass());
 	BrickSettings->OwnerTool = this;
+	CutterSettings->bNoisePreviewToggleEnabled = false; // TODO: support noise preview for bricks
 }
 
 FText UFractureToolBrick::GetDisplayText() const
diff --git a/Engine/Plugins/Experimental/ChaosEditor/Source/FractureEditor/Private/FractureToolCutter.cpp b/Engine/Plugins/Experimental/ChaosEditor/Source/FractureEditor/Private/FractureToolCutter.cpp
--- a/Engine/Plugins/Experimental/ChaosEditor/Source/FractureEditor/Private/FractureToolCutter.cpp
+++ b/Engine/Plugins/Experimental/ChaosEditor/Source/FractureEditor/Private/FractureToolCutter.cpp
@@ -14,9 +14,11 @@
 #include "GeometryCollection/GeometryCollectionObject.h"
 #include "GeometryCollection/GeometryCollectionComponent.h"
 #include "GeometryCollection/GeometryCollectionAlgo.h"
+#include "Components/DynamicMeshComponent.h"
 
 #include "Voronoi/Voronoi.h"
 #include "PlanarCut.h"
+#include "DynamicMesh/MeshTransforms.h"
 
 #include "FractureToolBackgroundTask.h"
 
@@ -240,6 +242,68 @@ void UFractureToolCutterBase::UpdateDefaultRandomSeed()
 }
 
 
+void FCellNoisePreviewOp::CalculateResult(FProgressCancel* Progress)
+{
+	Result->SetNum(ComputedDiagrams.Num());
+	for (int32 DiagramIdx = 0; DiagramIdx < ComputedDiagrams.Num(); ++DiagramIdx)
+	{
+		if (Progress && Progress->Cancelled())
+		{
+			return;
+		}
+
+		FCellDiagramResult& Info = ComputedDiagrams[DiagramIdx];
+		int32 NumCells = Info.Diagram.NumCells;
+		int32 KeepCellStep = 1.0f / FMath::Clamp(KeepCellsFrac, FLT_EPSILON, 1.0f);
+		auto KeepCellFunc = [KeepCellStep](int32 Idx) -> bool
+		{
+			return 0 == (Idx % KeepCellStep);
+		};
+		Info.Diagram.DiscardCells(KeepCellFunc, true);
+		FCellNoisePreviewResult& Preview = (*Result)[DiagramIdx];
+
+		if (Progress && Progress->Cancelled())
+		{
+			return;
+		}
+
+		CreateCuttingSurfacePreview(Info.Diagram, Info.Bounds, Grout, Info.NoiseSeed, Preview.NoiseMesh, KeepCellFunc, Info.Transform, Progress, Info.Transform.GetTranslation());
+
+		if (Progress && Progress->Cancelled())
+		{
+			return;
+		}
+
+		FVector Origin = Info.Transform.GetTranslation();
+		const UE::Geometry::FTransformSRT3d GeoTransform(Info.Transform);
+		MeshTransforms::ApplyTransform(Preview.NoiseMesh, GeoTransform, true);
+
+		Preview.BoneIdx = Info.BoneIdx;
+		Preview.Transform = Info.Transform;
+		Preview.SourceComponentIdx = Info.SourceComponentIdx;
+	}
+}
+
+void FVoronoiCellsOp::CalculateResult(FProgressCancel* Progress)
+{
+	for (FVoronoiDiagramInput& DiagramInfo : Diagrams)
+	{
+		if (Progress && Progress->Cancelled())
+		{
+			return;
+		}
+		FVector Origin = DiagramInfo.Transform.GetTranslation();
+		for (FVector& Site : DiagramInfo.Sites)
+		{
+			Site -= Origin;
+		}
+		DiagramInfo.Bounds.Max -= Origin;
+		DiagramInfo.Bounds.Min -= Origin;
+
+		FVoronoiDiagram Voronoi(DiagramInfo.Sites, DiagramInfo.Bounds, .1f);
+		Result->Emplace(FPlanarCells(DiagramInfo.Sites, Voronoi), DiagramInfo.SourceComponentIdx, DiagramInfo.BoneIdx, DiagramInfo.Transform, DiagramInfo.Bounds, DiagramInfo.NoiseSeed);
+	}
+}
 
 UFractureToolVoronoiCutterBase::UFractureToolVoronoiCutterBase(const FObjectInitializer& ObjInit)
 	: Super(ObjInit)
@@ -253,6 +317,69 @@ UFractureToolVoronoiCutterBase::UFractureToolVoronoiCutterBase(const FObjectInit
 	}
 }
 
+void UFractureToolVoronoiCutterBase::OnTick(float DeltaTime)
+{
+	ComputeRelaunchDelay -= DeltaTime;
+	if (bDiagramUpdated)
+	{
+		if (ComputeCells)
+		{
+			CancelBackgroundTask<FVoronoiCellsOp>(ComputeCells);
+			CancelBackgroundTask<FCellNoisePreviewOp>(ComputeNoisePreview);
+			ComputeRelaunchDelay = .5;
+		}
+		if (ComputeRelaunchDelay <= 0)
+		{
+			ComputeCells = StartBackgroundTask<FVoronoiCellsOp>(MakeUnique<FVoronoiCellsOp>(MoveTemp(DiagramInputs)));
+			bDiagramUpdated = false;
+		}
+	}
+	if (!CutterSettings->bDrawNoisePreview)
+	{
+		ClearMeshes();
+		CancelBackgroundTask<FCellNoisePreviewOp>(ComputeNoisePreview);
+	}
+	if (ComputeCells)
+	{
+		if (ComputeNoisePreview)
+		{
+			CancelBackgroundTask<FCellNoisePreviewOp>(ComputeNoisePreview);
+		}
+		TickBackgroundTask<FVoronoiCellsOp>(ComputeCells, false, [&](TUniquePtr<TArray<FCellDiagramResult>>&& Result) -> bool
+		{
+			AddLineVisualizations(*Result.Get());
+			ComputeRelaunchDelay = 0;
+			// TODO: Consider saving this result and passing it to the actual fracture Execute, rather than recomputing it (maybe not worth added complexity though)
+			TUniquePtr<FCellNoisePreviewOp> NoiseOp = MakeUnique<FCellNoisePreviewOp>(MoveTemp(*Result));
+
+			// Transfer noise settings
+			FNoiseSettings Settings;
+			CutterSettings->TransferNoiseSettings(Settings);
+			for (FCellDiagramResult& DiagramInfo : NoiseOp->ComputedDiagrams)
+			{
+				DiagramInfo.Diagram.SetNoise(Settings);
+			}
+			NoiseOp->PointSpacing = CollisionSettings->GetPointSpacing();
+			NoiseOp->Grout = CutterSettings->Grout;
+			NoiseOp->KeepCellsFrac = CutterSettings->FractionPreviewCells;
+			ensure(!ComputeNoisePreview);
+			if (CutterSettings->bDrawNoisePreview)
+			{
+				ComputeNoisePreview = StartBackgroundTask<FCellNoisePreviewOp>(MoveTemp(NoiseOp));
+			}
+			return true;
+		});
+	}
+	if (ComputeNoisePreview)
+	{
+		TickBackgroundTask<FCellNoisePreviewOp>(ComputeNoisePreview, false, [&](TUniquePtr<TArray<FCellNoisePreviewResult>>&& Result) -> bool
+		{
+			AddNoiseVisualizations(*Result.Get());
+			return true;
+		});
+	}
+}
+
 void UFractureToolVoronoiCutterBase::Render(const FSceneView* View, FViewport* Viewport, FPrimitiveDrawInterface* PDI)
 {
 	if (CutterSettings->bDrawSites)
@@ -268,6 +395,7 @@ void UFractureToolVoronoiCutterBase::Render(const FSceneView* View, FViewport* V
 
 void UFractureToolVoronoiCutterBase::UpdateLineSetExplodedVectors()
 {
+	bool bHasMeshes = VoronoiNoisePreviews.Num() == VoronoiLineSets.Num();
 	EnumerateVisualizationMapping(EdgesMappings, VoronoiLineSets.Num(), [&](int32 Idx, FVector ExplodedVector)
 	{
 		if (!VoronoiLineSets[Idx])
@@ -277,62 +405,127 @@ void UFractureToolVoronoiCutterBase::UpdateLineSetExplodedVectors()
 		// TODO: If we add diagrams even when visibility is toggled off, we also need to update visibility here, e.g.: VoronoiLineSets[Idx]->SetVisibility(CutterSettings->bDrawDiagram);
 		AActor* Actor = VoronoiLineSets[Idx]->GetOwner();
 		VoronoiLineSets[Idx]->SetRelativeLocation(Actor->GetTransform().InverseTransformVector(ExplodedVector));
+		if (bHasMeshes)
+		{
+			AActor* MeshActor = VoronoiNoisePreviews[Idx]->GetOwner();
+			VoronoiNoisePreviews[Idx]->SetRelativeLocation(MeshActor->GetTransform().InverseTransformVector(ExplodedVector));
+		}
 	});
 }
 
-void UFractureToolVoronoiCutterBase::UpdateVisualizations(TArray<FFractureToolContext>& FractureContexts)
+void UFractureToolVoronoiCutterBase::AddNoiseVisualizations(TArray<FCellNoisePreviewResult>& NoiseResults)
 {
-	ClearVisualizations();
+	ClearMeshes();
 
-	int32 MaxSitesToShowEdges = 100000; // computing all the voronoi diagrams can make the program non-responsive above this
-	int32 MaxSitesToShowSites = 1000000; // PDI struggles to render the site positions above this
-	bool bEstAboveMaxSites = false;
+	int32 NumCollections = NoiseResults.Num();
+	if (!ensure(VisualizedCollections.Num() == NumCollections))
+	{
+		return;
+	}
 
-	for (FFractureToolContext& FractureContext : FractureContexts)
+	if (!GEditor)
 	{
-		if (!FractureContext.GetBounds().IsValid) // skip contexts w/ invalid bounds
+		return;		
+	}
+	UWorld* World = GEditor->GetEditorWorldContext().World();
+
+	for (int32 CollectionIdx = 0; CollectionIdx < NumCollections; ++CollectionIdx)
+	{
+		UGeometryCollectionComponent* Component = VisualizedCollections[CollectionIdx].Get();
+		if (!ensure(Component))
 		{
 			continue;
 		}
-		int32 CollectionIdx = VisualizedCollections.Emplace(FractureContext.GetGeometryCollectionComponent());
-		int32 BoneIdx = FractureContext.GetSelection().Num() == 1 ? FractureContext.GetSelection()[0] : INDEX_NONE;
-		SitesMappings.AddMapping(CollectionIdx, BoneIdx, VoronoiSites.Num());
-		EdgesMappings.AddMapping(CollectionIdx, BoneIdx, VoronoiLineSets.Num());
-
-		// Generate voronoi diagrams and cache visualization info
-		TArray<FVector> LocalVoronoiSites;
-		GenerateVoronoiSites(FractureContext, LocalVoronoiSites);
-		// if diagram(s) become too large, skip the visualization
-		if (LocalVoronoiSites.Num() * FractureContexts.Num() > MaxSitesToShowSites || VoronoiSites.Num() + LocalVoronoiSites.Num() > MaxSitesToShowSites)
+		FCellNoisePreviewResult& Res = NoiseResults[CollectionIdx];
+
+		UDynamicMeshComponent* DynamicMeshComponent = NewObject<UDynamicMeshComponent>(Component);
+		DynamicMeshComponent->SetSceneProxyVerifyUsedMaterials(false);
+		DynamicMeshComponent->RegisterComponent();
+		VoronoiNoisePreviews.Add(DynamicMeshComponent);
+		DynamicMeshComponent->GetMesh()->Copy(Res.NoiseMesh);
+		DynamicMeshComponent->NotifyMeshUpdated();
+
+		// TODO: this material-access logic can go to the base class and be used by FractureToolPlaneCut as well
+		// Note the ToolSetupUtil functions optionally take a UInteractiveToolManager to provide fallback materials,
+		// but fracture mode does not set that up, so they may return a null material.
+		UMaterialInterface* NoiseMaterial = ToolSetupUtil::GetDefaultSculptMaterial(nullptr);
+		if (!NoiseMaterial)
 		{
-			UE_LOG(LogFractureTool, Warning, TEXT("Voronoi diagram(s) number of sites too large; will not display Voronoi diagram sites"));
-			ClearVisualizations();
-			break;
+			NoiseMaterial = ToolSetupUtil::GetDefaultMaterial();
 		}
-		VoronoiSites.Append(LocalVoronoiSites);
-		if (bEstAboveMaxSites || LocalVoronoiSites.Num() * FractureContexts.Num() > MaxSitesToShowEdges || VoronoiSites.Num() > MaxSitesToShowEdges)
+		DynamicMeshComponent->SetNumMaterials(1);
+		DynamicMeshComponent->SetMaterial(0, NoiseMaterial);
+	}
+}
+
+void UFractureToolVoronoiCutterBase::AddLineVisualizations(TArray<FCellDiagramResult>& DiagramResults)
+{
+	ClearEdges();
+
+	int32 NumCollections = DiagramResults.Num();
+	if (!ensure(VisualizedCollections.Num() == NumCollections))
+	{
+		RestoreEditorViewFlags();
+		return;
+	}
+
+	int32 TotalSites = 0;
+	for (const FCellDiagramResult& Result : DiagramResults)
+	{
+		TotalSites += Result.Diagram.NumCells;
+	}
+	// even though the diagrams are computed in the background, creating the line set component w/ all the voronoi diagrams edges causes a large hitch
+	// if the number of sites becomes too large, and the becomes diagram unreadably dense, so it's better to not draw it
+	constexpr int32 MaxSitesToShowEdges = 100000;
+	if (TotalSites > MaxSitesToShowEdges)
+	{
+		UE_LOG(LogFractureTool, Warning, TEXT("Voronoi diagram(s) number of sites too large; will not display Voronoi diagram edges"));
+		RestoreEditorViewFlags();
+		return;
+	}
+	for (int32 CollectionIdx = 0; CollectionIdx < DiagramResults.Num(); ++CollectionIdx)
+	{
+		const FCellDiagramResult& Result = DiagramResults[CollectionIdx];
+		check(CollectionIdx == Result.SourceComponentIdx);
+		EdgesMappings.AddMapping(CollectionIdx, Result.BoneIdx, VoronoiLineSets.Num());
+		UGeometryCollectionComponent* Component = VisualizedCollections[CollectionIdx].Get();
+		if (!ensure(Component))
 		{
-			UE_LOG(LogFractureTool, Warning, TEXT("Voronoi diagram(s) number of sites too large; will not display Voronoi diagram edges"));
-			ClearEdges();
-			bEstAboveMaxSites = true;
+			continue;
 		}
-		else if (CutterSettings->bDrawDiagram)
+		int32 NumBoundaries = Result.Diagram.PlaneBoundaries.Num();
+		if (CutterSettings->bDrawDiagram && NumBoundaries > 0)
 		{
-			FBox VoronoiBounds = GetVoronoiBounds(FractureContext, LocalVoronoiSites);
-			ULineSetComponent* VoronoiLineSet = NewObject<ULineSetComponent>(FractureContext.GetGeometryCollectionComponent());
-			VoronoiLineSet->SetupAttachment(FractureContext.GetGeometryCollectionComponent());
+			ULineSetComponent* VoronoiLineSet = NewObject<ULineSetComponent>(Component);
+			VoronoiLineSet->SetupAttachment(Component);
 			VoronoiLineSet->SetLineMaterial(ToolSetupUtil::GetDefaultLineComponentMaterial(nullptr, false));
-			// TODO: If we switch to adding diagrams even when visibility is toggled off, we also need to update visibility here, e.g.: VoronoiLineSets[Idx]->SetVisibility(CutterSettings->bDrawDiagram);
+			// TODO: If we switch to adding diagrams even when visibility is toggled off, we also need to update visibility here, 
+			// e.g.: VoronoiLineSet->SetVisibility(CutterSettings->bDrawDiagram);
 			VoronoiLineSet->RegisterComponent();
-			TArray<TTuple<FVector, FVector>> VoronoiEdges;
-			GetVoronoiEdges(LocalVoronoiSites, VoronoiBounds, VoronoiEdges, CellMember);
-			FTransform ToWorld = FractureContext.GetTransform();
-			VoronoiLineSet->ReserveLines(VoronoiEdges.Num());
-			for (int32 Idx = 0; Idx < VoronoiEdges.Num(); ++Idx)
+			int32 NumLines = 0;
+			for (int32 BoundaryIdx = 0; BoundaryIdx < NumBoundaries; ++BoundaryIdx)
 			{
-				const TTuple<FVector, FVector>& Line = VoronoiEdges[Idx];
-				VoronoiLineSet->AddLine(ToWorld.InverseTransformPosition(VoronoiEdges[Idx].Key), ToWorld.InverseTransformPosition(VoronoiEdges[Idx].Value), Colors[CellMember[Idx] % 100], 1.3, .1);
+				NumLines += Result.Diagram.PlaneBoundaries[BoundaryIdx].Num();
 			}
+			int32 LinesPerIndexHint = NumLines / NumBoundaries + 1;
+			FVector Origin = Result.Transform.GetTranslation();
+			VoronoiLineSet->AddLines(
+				NumBoundaries,
+				[&](int32 BoundaryIdx, TArray<FRenderableLine>& LinesOut)
+				{
+					const TArray<int32>& Boundary = Result.Diagram.PlaneBoundaries[BoundaryIdx];
+					TPair<int32, int32> Cells = Result.Diagram.PlaneCells[BoundaryIdx];
+					int32 Cell = Cells.Key > -1 ? Cells.Key : Cells.Value;
+					for (int32 PrevIdx = Boundary.Num() - 1, Idx = 0; Idx < Boundary.Num(); PrevIdx = Idx++)
+					{
+						FVector A = Result.Transform.InverseTransformPosition(Result.Diagram.PlaneBoundaryVertices[Boundary[PrevIdx]] + Origin);
+						FVector B = Result.Transform.InverseTransformPosition(Result.Diagram.PlaneBoundaryVertices[Boundary[Idx]] + Origin);
+
+						LinesOut.Emplace(A, B, Colors[Cell % 100], 1.3, .1);
+					}
+				},
+				LinesPerIndexHint,
+				false);
 			VoronoiLineSets.Add(VoronoiLineSet);
 		}
 	}
@@ -348,6 +541,51 @@ void UFractureToolVoronoiCutterBase::UpdateVisualizations(TArray<FFractureToolCo
 	}
 }
 
+void UFractureToolVoronoiCutterBase::UpdateVisualizations(TArray<FFractureToolContext>& FractureContexts)
+{
+	ClearVisualizations();
+
+	int32 MaxSitesToShowSites = 1000000; // PDI struggles to render the site positions above this
+	bool bEstAboveMaxSites = false;
+
+	DiagramInputs.Reset(FractureContexts.Num());
+	bDiagramUpdated = true;
+
+	for (FFractureToolContext& FractureContext : FractureContexts)
+	{
+		if (!FractureContext.GetBounds().IsValid) // skip contexts w/ invalid bounds
+		{
+			continue;
+		}
+		int32 CollectionIdx = VisualizedCollections.Emplace(FractureContext.GetGeometryCollectionComponent());
+		int32 BoneIdx = FractureContext.GetSelection().Num() == 1 ? FractureContext.GetSelection()[0] : INDEX_NONE;
+		SitesMappings.AddMapping(CollectionIdx, BoneIdx, VoronoiSites.Num());
+
+		// Store the voronoi diagrams we'd like to preview; Tick() will launch the background thread to compute them
+		TArray<FVector> LocalVoronoiSites;
+		GenerateVoronoiSites(FractureContext, LocalVoronoiSites);
+
+		FVoronoiDiagramInput& Diagram = DiagramInputs.Emplace_GetRef();
+		Diagram.Bounds = GetVoronoiBounds(FractureContext, LocalVoronoiSites);
+		Diagram.Sites = LocalVoronoiSites;
+		Diagram.SourceComponentIdx = CollectionIdx;
+		Diagram.Transform = FractureContext.GetTransform();
+		Diagram.BoneIdx = BoneIdx;
+		Diagram.NoiseSeed = FractureContext.GetSeed();
+
+		// Voronoi site point visualization doesn't need background compute, but is skipped if there are too many sites
+		if (LocalVoronoiSites.Num() * FractureContexts.Num() > MaxSitesToShowSites || VoronoiSites.Num() + LocalVoronoiSites.Num() > MaxSitesToShowSites)
+		{
+			UE_LOG(LogFractureTool, Warning, TEXT("Voronoi diagram(s) number of sites too large; will not display Voronoi diagram sites"));
+			ClearVisualizations();
+			break;
+		}
+		VoronoiSites.Append(LocalVoronoiSites);
+	}
+
+	RestoreEditorViewFlags();
+}
+
 void UFractureToolVoronoiCutterBase::FractureContextChanged()
 {
 	UpdateDefaultRandomSeed();
@@ -441,6 +679,19 @@ int32 UFractureToolVoronoiCutterBase::ExecuteFracture(const FFractureToolContext
 	return INDEX_NONE;
 }
 
+void UFractureToolVoronoiCutterBase::ClearMeshes()
+{
+	for (UDynamicMeshComponent* NoisePreview : VoronoiNoisePreviews)
+	{
+		if (NoisePreview)
+		{
+			NoisePreview->UnregisterComponent();
+			NoisePreview->DestroyComponent();
+		}
+	}
+	VoronoiNoisePreviews.Empty();
+}
+
 FBox UFractureToolVoronoiCutterBase::GetVoronoiBounds(const FFractureToolContext& FractureContext, const TArray<FVector>& Sites) const
 {
 	FBox VoronoiBounds = FractureContext.GetWorldBounds(); 
diff --git a/Engine/Plugins/Experimental/ChaosEditor/Source/FractureEditor/Private/FractureToolCutter.h b/Engine/Plugins/Experimental/ChaosEditor/Source/FractureEditor/Private/FractureToolCutter.h
--- a/Engine/Plugins/Experimental/ChaosEditor/Source/FractureEditor/Private/FractureToolCutter.h
+++ b/Engine/Plugins/Experimental/ChaosEditor/Source/FractureEditor/Private/FractureToolCutter.h
@@ -8,10 +8,14 @@
 #include "BaseGizmos/TransformProxy.h"
 #include "Drawing/LineSetComponent.h"
 
+// for background compute of preview geometry
+#include "PlanarCut.h"
+#include "FractureToolBackgroundTask.h"
+
 #include "FractureToolCutter.generated.h"
 
 class FFractureToolContext;
-
+class UDynamicMeshComponent;
 struct FNoiseSettings;
 
 /** Settings specifically related to the one-time destructive fracturing of a mesh **/
@@ -55,8 +59,16 @@ public:
 	UPROPERTY()
 	bool bGroupFractureToggleEnabled = true;
 
+	/** Amount of space to leave between cut pieces */
+	UPROPERTY(EditAnywhere, Category = CommonFracture, meta = (UIMin = "0.0", ClampMin = "0.0", EditCondition = "bGroutSettingEnabled", HideEditConditionToggle, EditConditionHides))
+	float Grout = 0.0f;
+
+	// This flag allows tools to disable the above Grout setting if/when it's not applicable
+	UPROPERTY()
+	bool bGroutSettingEnabled = true;
+
 	/** Draw points marking the centers of pieces to be cut out by the fracture pattern.  */
-	UPROPERTY(EditAnywhere, Category = CommonFracture, meta = (EditCondition = "bDrawSitesToggleEnabled", HideEditConditionToggle, EditConditionHides, DisplayName = "Draw Sites"))
+	UPROPERTY(EditAnywhere, Category = Visualization, meta = (EditCondition = "bDrawSitesToggleEnabled", HideEditConditionToggle, EditConditionHides, DisplayName = "Draw Sites"))
 	bool bDrawSites;
 
 	// This flag allows tools to disable the above bDrawSites option if/when it is not applicable
@@ -64,16 +76,28 @@ public:
 	bool bDrawSitesToggleEnabled = true;
 
 	/** Draw the edges of the fracture pattern.  */
-	UPROPERTY(EditAnywhere, Category = CommonFracture, meta = (DisplayName = "Draw Diagram"))
+	UPROPERTY(EditAnywhere, Category = Visualization, meta = (DisplayName = "Draw Diagram"))
 	bool bDrawDiagram;
 
-	/** Amount of space to leave between cut pieces */
-	UPROPERTY(EditAnywhere, Category = CommonFracture, meta = (UIMin = "0.0", ClampMin = "0.0", EditCondition = "bGroutSettingEnabled", HideEditConditionToggle, EditConditionHides))
-	float Grout = 0.0f;
+	/** Whether to show a solid preview of the cutting geometry, including any noise displacement */
+	UPROPERTY(EditAnywhere, Category = Visualization, meta = (EditCondition = "bNoisePreviewToggleEnabled", HideEditConditionToggle, EditConditionHides))
+	bool bDrawNoisePreview = false;
 
-	// This flag allows tools to disable the above Noise settings if/when they are not applicable
+	// This flag allows tools to disable the above bDrawNoisePreview option if/when it is not applicable
 	UPROPERTY()
-	bool bGroutSettingEnabled = true;
+	bool bNoisePreviewToggleEnabled = true;
+
+	/** Fraction of cells to show in noise preview */
+	UPROPERTY(EditAnywhere, Category = Visualization, meta = (ClampMin = "0.0", ClampMax = "1.0", EditCondition = "bDrawNoisePreview && bDrawSitesToggleEnabled", HideEditConditionToggle, EditConditionHides))
+	float FractionPreviewCells = .1;
+
+	/** Scale of the noise preview plane */
+	UPROPERTY(EditAnywhere, Category = Visualization, meta = (ClampMin = ".01", ClampMax = "1000", UIMin = ".25", UIMax = "10", EditCondition = "bDrawNoisePreview && bNoisePreviewHasScale", HideEditConditionToggle, EditConditionHides))
+	double NoisePreviewScale = 1.;
+
+	// This flag allows tools to disable the noise preview scale setting if it's not applicable
+	UPROPERTY()
+	bool bNoisePreviewHasScale = false;
 
 	/** Size of the Perlin noise displacement (in cm). If 0, no noise will be applied */
 	UPROPERTY(EditAnywhere, Category = Noise, meta = (UIMin = "0.0", EditCondition = "bNoiseSettingsEnabled", HideEditConditionToggle, EditConditionHides))
@@ -218,6 +242,71 @@ protected:
 };
 
 
+struct FVoronoiDiagramInput
+{
+	FTransform Transform;
+	int32 SourceComponentIdx;
+	int32 BoneIdx;
+	TArray<FVector> Sites;
+	FBox Bounds;
+	int32 NoiseSeed;
+};
+
+struct FCellDiagramResult
+{
+	FCellDiagramResult(FPlanarCells Diagram, int32 SourceComponentIdx, int32 BoneIdx, FTransform Transform, FBox Bounds, int32 NoiseSeed)
+		: Diagram(Diagram), SourceComponentIdx(SourceComponentIdx), BoneIdx(BoneIdx), Transform(Transform), Bounds(Bounds), NoiseSeed(NoiseSeed)
+	{}
+
+	FPlanarCells Diagram;
+	int32 SourceComponentIdx;
+	int32 BoneIdx;
+	FTransform Transform;
+	FBox Bounds;
+	int32 NoiseSeed;
+};
+
+struct FCellNoisePreviewResult
+{
+	UE::Geometry::FDynamicMesh3 NoiseMesh;
+	int32 SourceComponentIdx;
+	int32 BoneIdx;
+	FTransform Transform;
+};
+
+class FCellNoisePreviewOp : public UE::Geometry::TGenericDataOperator<TArray<FCellNoisePreviewResult>>
+{
+public:
+	FCellNoisePreviewOp(TArray<FCellDiagramResult>&& ComputedDiagrams) : ComputedDiagrams(MoveTemp(ComputedDiagrams))
+	{}
+	virtual ~FCellNoisePreviewOp() = default;
+
+	TArray<FCellDiagramResult> ComputedDiagrams;
+	float PointSpacing;
+	float Grout;
+
+	float KeepCellsFrac = .1;
+
+	// TGenericDataOperator interface:
+	virtual void CalculateResult(FProgressCancel* Progress) override;
+};
+
+class FVoronoiCellsOp : public UE::Geometry::TGenericDataOperator<TArray<FCellDiagramResult>>
+{
+public:
+
+	FVoronoiCellsOp(TArray<FVoronoiDiagramInput>&& Inputs) : Diagrams(MoveTemp(Inputs))
+	{}
+
+	virtual ~FVoronoiCellsOp() = default;
+
+	TArray<FVoronoiDiagramInput> Diagrams;
+
+	// TGenericDataOperator interface:
+	virtual void CalculateResult(FProgressCancel* Progress) override;
+};
+
+
 UCLASS(Abstract, DisplayName = "Voronoi Base", Category = "FractureTools")
 class UFractureToolVoronoiCutterBase : public UFractureToolCutterBase
 {
@@ -226,6 +315,7 @@ public:
 
 	UFractureToolVoronoiCutterBase(const FObjectInitializer& ObjInit);
 
+	virtual void OnTick(float DeltaTime) override;
 	virtual void Render(const FSceneView* View, FViewport* Viewport, FPrimitiveDrawInterface* PDI) override;
 
 	virtual void FractureContextChanged() override;
@@ -242,6 +332,13 @@ protected:
 		VoronoiSites.Empty();
 		SitesMappings.Empty();
 		ClearEdges();
+		ClearMeshes();
+		if (ComputeCells)
+		{
+			UE::Fracture::CancelBackgroundTask<FVoronoiCellsOp>(ComputeCells);
+			UE::Fracture::CancelBackgroundTask<FCellNoisePreviewOp>(ComputeNoisePreview);
+			ComputeRelaunchDelay = .5;
+		}
 	}
 
 	void ClearEdges()
@@ -258,11 +355,27 @@ protected:
 		VoronoiLineSets.Empty();
 	}
 
+	void ClearMeshes();
+
 	virtual void UpdateVisualizations(TArray<FFractureToolContext>& FractureContexts);
+	virtual void AddLineVisualizations(TArray<FCellDiagramResult>& DiagramResults);
+	virtual void AddNoiseVisualizations(TArray<FCellNoisePreviewResult>& NoiseResults);
 
 private:
 	UPROPERTY()
 	TArray<TObjectPtr<ULineSetComponent>> VoronoiLineSets;
+	TArray<TObjectPtr<UDynamicMeshComponent>> VoronoiNoisePreviews;
+
+
+	//
+	// Background compute
+	//
+	TUniquePtr<UE::Fracture::TBackgroundOpExecuter<FVoronoiCellsOp>> ComputeCells = nullptr;
+	TUniquePtr<UE::Fracture::TBackgroundOpExecuter<FCellNoisePreviewOp>> ComputeNoisePreview = nullptr;
+
+	TArray<FVoronoiDiagramInput> DiagramInputs;
+	bool bDiagramUpdated = false;
+	float ComputeRelaunchDelay = 0;
 
 	TArray<int32> CellMember;
 	TArray<FVector> VoronoiSites;
diff --git a/Engine/Plugins/Experimental/ChaosEditor/Source/FractureEditor/Private/FractureToolMeshCut.cpp b/Engine/Plugins/Experimental/ChaosEditor/Source/FractureEditor/Private/FractureToolMeshCut.cpp
--- a/Engine/Plugins/Experimental/ChaosEditor/Source/FractureEditor/Private/FractureToolMeshCut.cpp
+++ b/Engine/Plugins/Experimental/ChaosEditor/Source/FractureEditor/Private/FractureToolMeshCut.cpp
@@ -29,6 +29,7 @@ UFractureToolMeshCut::UFractureToolMeshCut(const FObjectInitializer& ObjInit)
 	MeshCutSettings = NewObject<UFractureMeshCutSettings>(GetTransientPackage(), UFractureMeshCutSettings::StaticClass());
 	MeshCutSettings->OwnerTool = this;
 	CutterSettings->bDrawSitesToggleEnabled = false;
+	CutterSettings->bNoisePreviewToggleEnabled = false;
 	DisableNoiseSettings();
 	DisableGroutSetting();
 }
diff --git a/Engine/Plugins/Experimental/ChaosEditor/Source/FractureEditor/Private/FractureToolPlaneCut.cpp b/Engine/Plugins/Experimental/ChaosEditor/Source/FractureEditor/Private/FractureToolPlaneCut.cpp
--- a/Engine/Plugins/Experimental/ChaosEditor/Source/FractureEditor/Private/FractureToolPlaneCut.cpp
+++ b/Engine/Plugins/Experimental/ChaosEditor/Source/FractureEditor/Private/FractureToolPlaneCut.cpp
@@ -123,7 +123,7 @@ TUniquePtr<UE::Geometry::FDynamicMeshOperator> UFractureToolPlaneCut::MakeNewOpe
 				NoisePreviewOp->NoisePivots.Add(NoisePivots[Idx]);
 			});
 	}
-	NoisePreviewOp->PlaneSize *= PlaneCutSettings->NoisePreviewScale;
+	NoisePreviewOp->PlaneSize *= CutterSettings->NoisePreviewScale;
 
 	return NoisePreviewOp;
 }
@@ -146,6 +146,9 @@ void UFractureToolPlaneCut::Setup()
 	GizmoSettings->Setup(this);
 	PlaneCutSettings->bCanCutWithMultiplePlanes = !GizmoSettings->bUseGizmo;
 	NotifyOfPropertyChangeByTool(PlaneCutSettings);
+	CutterSettings->bDrawSitesToggleEnabled = false;
+	CutterSettings->bNoisePreviewHasScale = true;
+	CutterSettings->bDrawNoisePreview = true; // default-enable the plane noise preview
 
 	// Initialize the background compute object for the noise preview
 	if (GEditor && !NoisePreview)
@@ -322,7 +325,7 @@ void UFractureToolPlaneCut::FractureContextChanged()
 
 	if (NoisePreview)
 	{
-		NoisePreview->SetVisibility(PlaneCutSettings->bShowNoisePreview);
+		NoisePreview->SetVisibility(CutterSettings->bDrawNoisePreview);
 	}
 
 	RenderCuttingPlaneSize = FLT_MAX;
diff --git a/Engine/Plugins/Experimental/ChaosEditor/Source/FractureEditor/Private/FractureToolPlaneCut.h b/Engine/Plugins/Experimental/ChaosEditor/Source/FractureEditor/Private/FractureToolPlaneCut.h
--- a/Engine/Plugins/Experimental/ChaosEditor/Source/FractureEditor/Private/FractureToolPlaneCut.h
+++ b/Engine/Plugins/Experimental/ChaosEditor/Source/FractureEditor/Private/FractureToolPlaneCut.h
@@ -29,14 +29,6 @@ public:
 	UPROPERTY()
 	bool bCanCutWithMultiplePlanes = false;
 
-	/** Whether to show the noise displacement on the preview visualization of the cutting planes */
-	UPROPERTY(EditAnywhere, Category = Visualization)
-	bool bShowNoisePreview = true;
-
-	/** Scale of the noise preview plane */
-	UPROPERTY(EditAnywhere, Category = Visualization, meta = (ClampMin = ".01", ClampMax = "1000", UIMin = ".25", UIMax = "10", EditCondition = "bShowNoisePreview", HideEditConditionToggle))
-	double NoisePreviewScale = 1.;
-
 };
 
 
diff --git a/Engine/Plugins/Experimental/ChaosEditor/Source/FractureEditor/Public/FractureTool.h b/Engine/Plugins/Experimental/ChaosEditor/Source/FractureEditor/Public/FractureTool.h
--- a/Engine/Plugins/Experimental/ChaosEditor/Source/FractureEditor/Public/FractureTool.h
+++ b/Engine/Plugins/Experimental/ChaosEditor/Source/FractureEditor/Public/FractureTool.h
@@ -211,7 +211,7 @@ protected:
 
 	// Geometry collection components referenced by visualizations
 	UPROPERTY()
-	TArray<TObjectPtr<const UGeometryCollectionComponent>> VisualizedCollections;
+	TArray<TObjectPtr<UGeometryCollectionComponent>> VisualizedCollections;
 
 	virtual void ClearVisualizations()
 	{
diff --git a/Engine/Plugins/Experimental/PlanarCutPlugin/Source/PlanarCut/Private/PlanarCut.cpp b/Engine/Plugins/Experimental/PlanarCutPlugin/Source/PlanarCut/Private/PlanarCut.cpp
--- a/Engine/Plugins/Experimental/PlanarCutPlugin/Source/PlanarCut/Private/PlanarCut.cpp
+++ b/Engine/Plugins/Experimental/PlanarCutPlugin/Source/PlanarCut/Private/PlanarCut.cpp
@@ -41,6 +41,7 @@
 #include "DynamicMeshToMeshDescription.h"
 
 #include "Algo/Rotate.h"
+#include "Algo/Reverse.h"
 
 #include "GeometryMeshConversion.h"
 
@@ -717,6 +718,89 @@ FPlanarCells::FPlanarCells(const FBox &Region, const TArrayView<const FColor> Im
 }
 
 
+void FPlanarCells::DiscardCells(TFunctionRef<bool(int32)> KeepFunc, bool bKeepNeighbors)
+{
+	TArray<int32> OldToNew;
+	OldToNew.Init(-1, NumCells);
+	int32 KeptCells = 0;
+	for (int32 CellIdx = 0; CellIdx < NumCells; ++CellIdx)
+	{
+		if (KeepFunc(CellIdx))
+		{
+			OldToNew[CellIdx] = KeptCells++;
+		}
+	}
+	if (bKeepNeighbors && KeptCells < NumCells)
+	{
+		for (const TPair<int32, int32>& Neighbors : PlaneCells)
+		{
+			if (Neighbors.Key < 0 || Neighbors.Value < 0)
+			{
+				continue;
+			}
+			bool bKeptKey = KeepFunc(Neighbors.Key);
+			bool bKeptValue = KeepFunc(Neighbors.Value);
+			if (bKeptKey != bKeptValue)
+			{
+				int32 Nbr = bKeptKey ? Neighbors.Value : Neighbors.Key;
+				if (OldToNew[Nbr] < 0)
+				{
+					OldToNew[Nbr] = KeptCells++;
+				}
+			}
+		}
+	}
+	if (KeptCells == NumCells)
+	{
+		return;
+	}
+
+	NumCells = KeptCells;
+	for (int32 PlaneIdx = 0; PlaneIdx < Planes.Num(); ++PlaneIdx)
+	{
+		TPair<int32, int32> Cells = PlaneCells[PlaneIdx];
+		Cells.Key = Cells.Key > -1 ? OldToNew[Cells.Key] : -1;
+		Cells.Value = Cells.Value > -1 ? OldToNew[Cells.Value] : -1;
+		if (Cells.Key == Cells.Value && Cells.Key == -1)
+		{
+			PlaneCells.RemoveAtSwap(PlaneIdx, 1, false);
+			Planes.RemoveAtSwap(PlaneIdx, 1, false);
+			PlaneBoundaries.RemoveAtSwap(PlaneIdx, 1, false);
+			PlaneIdx--; // consider the swapped-in value in the next iteration
+		}
+		else
+		{
+			// on boundary to outside, the 'outside' index must always be second
+			// if the discards above broke that invariant, flip the plane to fix it
+			if (Cells.Key < 0)
+			{
+				Swap(Cells.Key, Cells.Value);
+				Algo::Reverse(PlaneBoundaries[PlaneIdx]);
+				Planes[PlaneIdx] = Planes[PlaneIdx].Flip();
+			}
+			PlaneCells[PlaneIdx] = Cells;
+		}
+	}
+
+	// Compress vertices array to only the used vertices
+	TArray<int32> OldToNewVertex;
+	OldToNewVertex.Init(-1, PlaneBoundaryVertices.Num());
+	TArray<FVector> NewBoundaryVertices;
+	for (TArray<int32>& PlaneBoundary : PlaneBoundaries)
+	{
+		for (int32& VID : PlaneBoundary)
+		{
+			int32& NewVID = OldToNewVertex[VID];
+			if (NewVID < 0)
+			{
+				NewVID = NewBoundaryVertices.Add(PlaneBoundaryVertices[VID]);
+			}
+			VID = NewVID;
+		}
+	}
+	PlaneBoundaryVertices = MoveTemp(NewBoundaryVertices);
+}
+
 
 
 // Simpler invocation of CutWithPlanarCells w/ reasonable defaults
@@ -806,6 +890,51 @@ int32 CutMultipleWithMultiplePlanes(
 	return NewGeomStartIdx;
 }
 
+void CreateCuttingSurfacePreview(
+	const FPlanarCells& Cells,
+	const FBox& Bounds,
+	double Grout,
+	int32 RandomSeed,
+	FDynamicMesh3& OutCuttingMeshes,
+	TFunctionRef<bool(int)> FilterCellsFunc,
+	const TOptional<FTransform>& TransformCollection,
+	FProgressCancel* Progress,
+	FVector CellsOrigin
+)
+{
+	FTransform CollectionToWorld = TransformCollection.Get(FTransform::Identity);
+	// Put Collection in the same local space as the Cells
+	FTransform CollectionToWorldCentered = CollectionToWorld * FTransform(-CellsOrigin);
+
+	FProgressCancel::FProgressScope SurfaceScope = FProgressCancel::CreateScopeTo(Progress, .99);
+	UE::Geometry::FAxisAlignedBox3d GeoBounds(Bounds);
+	double OnePercentExtend = GeoBounds.MaxDim() * .01;
+	FRandomStream RandomStream(RandomSeed);
+	FCellMeshes CellMeshes(0, RandomStream, Cells, GeoBounds, Grout, OnePercentExtend, false);
+	SurfaceScope.Done();
+
+	if (Progress && Progress->Cancelled())
+	{
+		return;
+	}
+
+	FProgressCancel::FProgressScope AppendScope = FProgressCancel::CreateScopeTo(Progress, 1.0);
+	OutCuttingMeshes.Clear();
+	OutCuttingMeshes.EnableAttributes();
+	FDynamicMeshEditor Editor(&OutCuttingMeshes);
+
+	FMeshIndexMappings IndexMaps; // Needed for Editor.AppendMesh, not used otherwise
+	for (int32 CellIdx = 0; CellIdx < CellMeshes.CellMeshes.Num(); ++CellIdx)
+	{
+		if (FilterCellsFunc(CellIdx))
+		{
+			Editor.AppendMesh(&CellMeshes.CellMeshes[CellIdx].AugMesh, IndexMaps);
+		}
+	}
+	
+	AppendScope.Done();
+}
+
 // Cut multiple Geometry groups inside a GeometryCollection with PlanarCells, and add each cut cell back to the GeometryCollection as a new child of their source Geometry
 int32 CutMultipleWithPlanarCells(
 	FPlanarCells& Cells,
diff --git a/Engine/Plugins/Experimental/PlanarCutPlugin/Source/PlanarCut/Public/PlanarCut.h b/Engine/Plugins/Experimental/PlanarCutPlugin/Source/PlanarCut/Public/PlanarCut.h
--- a/Engine/Plugins/Experimental/PlanarCutPlugin/Source/PlanarCut/Public/PlanarCut.h
+++ b/Engine/Plugins/Experimental/PlanarCutPlugin/Source/PlanarCut/Public/PlanarCut.h
@@ -110,6 +110,9 @@ struct PLANARCUT_API FPlanarCells
 
 	FInternalSurfaceMaterials InternalSurfaceMaterials;
 
+	// Discards cells from the diagram, used to support visualization of a subset of cells
+	void DiscardCells(TFunctionRef<bool(int32)> KeepFunc, bool bKeepNeighbors);
+
 	/**
 	 * @return true if this is a single, unbounded cutting plane
 	 */
@@ -216,6 +219,31 @@ int32 PLANARCUT_API CutWithPlanarCells(
 	FVector CellsOrigin = FVector::ZeroVector
 );
 
+/**
+ * Generate cutting mesh surfaces as a single mesh
+ * Useful for creating a preview of the cut surface.
+ * 
+ * @param Cells				Defines the cutting planes and division of space
+ * @param Grout				Separation to leave between cutting cells
+ * @param RandomSeed				Seed to be used for random noise displacement
+ * @param OutCuttingMeshes			Dynamic mesh representing the preview surface
+ * @param FilterCellsFunc			Filter which cells should be included in the preview
+ * @param TransformCollection		Optional transform that would be applied to the to-cut surface; if unset, defaults to Identity
+ * @param Progress						Optionally tracks progress and supports early-cancel
+ * @param CellsOrigin					Optionally provide a local origin of the cutting Cells
+ */
+void PLANARCUT_API CreateCuttingSurfacePreview(
+	const FPlanarCells& Cells,
+	const FBox& Bounds,
+	double Grout,
+	int32 RandomSeed,
+	UE::Geometry::FDynamicMesh3& OutCuttingMeshes,
+	TFunctionRef<bool(int)> FilterCellsFunc,
+	const TOptional<FTransform>& TransformCollection = TOptional<FTransform>(),
+	FProgressCancel* Progress = nullptr,
+	FVector CellsOrigin = FVector::ZeroVector
+);
+
 /**
  * Cut multiple Geometry groups inside a GeometryCollection with PlanarCells, and add each cut cell back to the GeometryCollection as a new child of their source Geometry.  For geometries that would not be cut, nothing is added.
  *
diff --git a/Engine/Plugins/Runtime/MeshModelingToolset/Source/ModelingOperators/Public/ModelingOperators.h b/Engine/Plugins/Runtime/MeshModelingToolset/Source/ModelingOperators/Public/ModelingOperators.h
--- a/Engine/Plugins/Runtime/MeshModelingToolset/Source/ModelingOperators/Public/ModelingOperators.h
+++ b/Engine/Plugins/Runtime/MeshModelingToolset/Source/ModelingOperators/Public/ModelingOperators.h
@@ -116,6 +116,8 @@ protected:
 	TUniquePtr<ResultType> Result;
 
 public:
+
+	using ResultTypeName = ResultType;
 	
 	TGenericDataOperator(bool bCreateInitialObject = true)
 	{
