rework how internal faces are tracked and assigned materials in geometry collections:
- add explicit tracking of internal surfaces in fracture geometry, and remove concept of 'internal-only' materials in odd-numbered slots.
- add tooling to set internal materials, and to set default internal materials to be assigned on fracture
- add option to preserve internal face data via odd material numbers in ToMesh tool + option to re-import these ToMesh outputs w/ internal material IDs set

diff --git a/Engine/Plugins/Experimental/ChaosEditor/Source/FractureEditor/Private/FractureEditorCommands.h b/Engine/Plugins/Experimental/ChaosEditor/Source/FractureEditor/Private/FractureEditorCommands.h
--- a/Engine/Plugins/Experimental/ChaosEditor/Source/FractureEditor/Private/FractureEditorCommands.h
+++ b/Engine/Plugins/Experimental/ChaosEditor/Source/FractureEditor/Private/FractureEditorCommands.h
@@ -75,6 +75,9 @@ class FFractureEditorCommands : public TCommands<FFractureEditorCommands>
 		// UV Commands
 		TSharedPtr< FUICommandInfo > AutoUV;
 		
+		// Material Commands
+		TSharedPtr< FUICommandInfo > Materials;
+
 		// Fracture Commands
 		TSharedPtr< FUICommandInfo > Uniform;
 		TSharedPtr< FUICommandInfo > Radial;
diff --git a/Engine/Plugins/Experimental/ChaosEditor/Source/FractureEditor/Private/FractureEditorModeToolkit.cpp b/Engine/Plugins/Experimental/ChaosEditor/Source/FractureEditor/Private/FractureEditorModeToolkit.cpp
--- a/Engine/Plugins/Experimental/ChaosEditor/Source/FractureEditor/Private/FractureEditorModeToolkit.cpp
+++ b/Engine/Plugins/Experimental/ChaosEditor/Source/FractureEditor/Private/FractureEditorModeToolkit.cpp
@@ -864,6 +864,7 @@ void FFractureEditorModeToolkit::BuildToolPalette(FName PaletteIndex, class FToo
 	else if (PaletteIndex == TEXT("Utilities"))
 	{
 		ToolbarBuilder.AddToolBarButton(Commands.AutoUV);
+		ToolbarBuilder.AddToolBarButton(Commands.Materials);
 		ToolbarBuilder.AddToolBarButton(Commands.RecomputeNormals);
 		ToolbarBuilder.AddToolBarButton(Commands.Resample);
 		ToolbarBuilder.AddToolBarButton(Commands.ConvertToMesh);
diff --git a/Engine/Plugins/Experimental/ChaosEditor/Source/FractureEditor/Private/FractureToolConvert.cpp b/Engine/Plugins/Experimental/ChaosEditor/Source/FractureEditor/Private/FractureToolConvert.cpp
--- a/Engine/Plugins/Experimental/ChaosEditor/Source/FractureEditor/Private/FractureToolConvert.cpp
+++ b/Engine/Plugins/Experimental/ChaosEditor/Source/FractureEditor/Private/FractureToolConvert.cpp
@@ -39,6 +39,8 @@ UFractureToolConvert::UFractureToolConvert(const FObjectInitializer& ObjInit)
 {
 	ConvertSettings = NewObject<UFractureConvertSettings>(GetTransientPackage(), UFractureConvertSettings::StaticClass());
 	ConvertSettings->OwnerTool = this;
+	ReimportSettings = NewObject<UFractureReimportSettings>(GetTransientPackage(), UFractureReimportSettings::StaticClass());
+	ReimportSettings->OwnerTool = this;
 }
 
 bool UFractureToolConvert::CanExecute() const
@@ -76,6 +78,7 @@ TArray<UObject*> UFractureToolConvert::GetSettingsObjects() const
 {
 	TArray<UObject*> Settings;
 	Settings.Add(ConvertSettings);
+	Settings.Add(ReimportSettings);
 	return Settings;
 }
 
@@ -139,7 +142,7 @@ namespace FractureToolConvertInternal
 // Create + position a single static mesh asset in the provided asset package
 void CreateMeshAsset(
 	FGeometryCollection& Collection, const TManagedArray<FTransform>& BoneTransforms, const FTransform& CollectionToWorld, const TConstArrayView<UMaterialInterface*>& Materials,
-	TConstArrayView<int32> Bones, UPackage* AssetPackage, FString UniqueAssetName, bool bCenterPivot, bool bPlaceInWorld, bool bSelectNewActors)
+	TConstArrayView<int32> Bones, UPackage* AssetPackage, FString UniqueAssetName, bool bCenterPivot, bool bPlaceInWorld, bool bSelectNewActors, bool bOddMaterialsAreInternal)
 {
 	// create new UStaticMesh object
 	EObjectFlags Flags = EObjectFlags::RF_Public | EObjectFlags::RF_Standalone;
@@ -157,9 +160,17 @@ void CreateMeshAsset(
 	NewStaticMesh->CreateBodySetup();
 	NewStaticMesh->GetBodySetup()->CollisionTraceFlag = ECollisionTraceFlag::CTF_UseComplexAsSimple;
 
+	TFunction<int32(int32, bool)> RemapMaterialIDs = nullptr;
+	if (bOddMaterialsAreInternal)
+	{
+		RemapMaterialIDs = [](int32 InMatID, bool bIsInternal)
+		{
+			return InMatID * 2 + (int32)bIsInternal;
+		};
+	}
 
 	FTransform BonesToCollection;
-	ConvertToMeshDescription(*OutputMeshDescription, BonesToCollection, bCenterPivot, Collection, BoneTransforms, Bones);
+	ConvertToMeshDescription(*OutputMeshDescription, BonesToCollection, bCenterPivot, Collection, BoneTransforms, Bones, RemapMaterialIDs);
 
 	// add a material slot. Must always have one material slot.
 	int AddMaterialCount = FMath::Max(1, Materials.Num());
@@ -217,7 +228,8 @@ void CreateMeshAsset(
 bool ConvertAndSaveMeshes(
 	FGeometryCollection& Collection, const TManagedArray<FTransform>& BoneTransforms, const FTransform& CollectionToWorld,
 	const TConstArrayView<UMaterialInterface*>& Materials, TConstArrayView<int32> Bones, FString ObjectBaseName, 
-	const UObject* RelativeToAsset, bool bPromptToSave, bool bSaveCombined, bool bCenterPivots, bool bPlaceInWorld, bool bSelectNewActors)
+	const UObject* RelativeToAsset, bool bPromptToSave, bool bSaveCombined, bool bCenterPivots, bool bPlaceInWorld,
+	bool bSelectNewActors, bool bOddMaterialsAreInternal)
 {
 	check(RelativeToAsset);
 
@@ -280,7 +292,7 @@ bool ConvertAndSaveMeshes(
 		FString UniqueAssetName;
 		UPackage* AssetPackage = MakeUniquePackage(ObjectBaseName, UniqueAssetName);
 		SavePackage.Add(AssetPackage);
-		CreateMeshAsset(Collection, BoneTransforms, CollectionToWorld, Materials, Bones, AssetPackage, UniqueAssetName, bCenterPivots, bPlaceInWorld, bSelectNewActors);
+		CreateMeshAsset(Collection, BoneTransforms, CollectionToWorld, Materials, Bones, AssetPackage, UniqueAssetName, bCenterPivots, bPlaceInWorld, bSelectNewActors, bOddMaterialsAreInternal);
 	}
 	else
 	{
@@ -298,7 +310,7 @@ bool ConvertAndSaveMeshes(
 
 			TArrayView<const int32> SingleBoneView(&UseBone, 1);
 
-			CreateMeshAsset(Collection, BoneTransforms, CollectionToWorld, Materials, SingleBoneView, AssetPackage, UniqueAssetName, bCenterPivots, bPlaceInWorld, bSelectNewActors);
+			CreateMeshAsset(Collection, BoneTransforms, CollectionToWorld, Materials, SingleBoneView, AssetPackage, UniqueAssetName, bCenterPivots, bPlaceInWorld, bSelectNewActors, bOddMaterialsAreInternal);
 		}
 	}
 
@@ -357,12 +369,25 @@ int32 UFractureToolConvert::ExecuteFracture(const FFractureToolContext& Fracture
 			}
 		}
 
+		// If we need odd materials as internal, naively duplicate materials
+		// Note: We could do this in a smarter pattern based to minimize material slots, but at the cost of having a less predictable / more confusing result
+		if (ReimportSettings->bOddMaterialsAreInternal)
+		{
+			int32 RealMaterialNum = Materials.Num();
+			Materials.SetNum(RealMaterialNum * 2);
+			for (int32 MaterialIndex = RealMaterialNum - 1; MaterialIndex >= 0; --MaterialIndex)
+			{
+				Materials[MaterialIndex * 2 + 1] = Materials[MaterialIndex];
+				Materials[MaterialIndex * 2] = Materials[MaterialIndex];
+			}
+		}
+
 		// choose default mesh name based on corresponding geometry collection name
 		FString BaseName = FString::Printf(TEXT("%s_SM"), *FractureContext.GetFracturedGeometryCollection()->GetName());
 		
 		UE::FractureToolConvertInternal::ConvertAndSaveMeshes(Collection, BoneTransforms, FractureContext.GetTransform(), Materials,
 			FractureContext.GetSelection(), BaseName, FractureContext.GetFracturedGeometryCollection(), 
-			ConvertSettings->bPromptForBaseName, !ConvertSettings->bPerBone, ConvertSettings->bCenterPivots, ConvertSettings->bPlaceInWorld, ConvertSettings->bSelectNewActors);
+			ConvertSettings->bPromptForBaseName, !ConvertSettings->bPerBone, ConvertSettings->bCenterPivots, ConvertSettings->bPlaceInWorld, ConvertSettings->bSelectNewActors, ReimportSettings->bOddMaterialsAreInternal);
 	}
 
 	return INDEX_NONE;
diff --git a/Engine/Plugins/Experimental/ChaosEditor/Source/FractureEditor/Private/FractureToolConvert.h b/Engine/Plugins/Experimental/ChaosEditor/Source/FractureEditor/Private/FractureToolConvert.h
--- a/Engine/Plugins/Experimental/ChaosEditor/Source/FractureEditor/Private/FractureToolConvert.h
+++ b/Engine/Plugins/Experimental/ChaosEditor/Source/FractureEditor/Private/FractureToolConvert.h
@@ -42,6 +42,22 @@ public:
 	bool bSelectNewActors = false;
 };
 
+/** Settings related to geometry collection -> static mesh conversion **/
+UCLASS(config = EditorPerProjectUserSettings)
+class UFractureReimportSettings : public UFractureToolSettings
+{
+public:
+	GENERATED_BODY()
+
+	// If checked, materials will be updated on export to ensure that odd-numbered material IDs correspond to internal faces, and this will be used on re-import to assign correct material IDs
+	UPROPERTY(EditAnywhere, Category = ReimportSettings)
+	bool bOddMaterialsAreInternal = true;
+
+	// TODO: Move the "re-import ToMesh output" path from out of the Generate/New UI an into this tool.  (Or provide alternate tooling and deprecate the need for this path.)
+	// Reimport a static mesh to the currently-selected geometry collection, or to a new geometry collection if none is selected.
+	//UFUNCTION(CallInEditor, Category = ReimportSettings, meta = (DisplayName = "Reimport to Selected"))
+	//void ReimportToSelected();
+};
 
 UCLASS(DisplayName = "Convert Tool", Category = "FractureTools")
 class UFractureToolConvert : public UFractureModalTool
@@ -83,4 +99,7 @@ public:
 protected:
 	UPROPERTY(EditAnywhere, Category = Slicing)
 	TObjectPtr<UFractureConvertSettings> ConvertSettings;
+
+	UPROPERTY(EditAnywhere, Category = Slicing)
+	TObjectPtr<UFractureReimportSettings> ReimportSettings;
 };
diff --git a/Engine/Plugins/Experimental/ChaosEditor/Source/FractureEditor/Private/FractureToolCutter.cpp b/Engine/Plugins/Experimental/ChaosEditor/Source/FractureEditor/Private/FractureToolCutter.cpp
--- a/Engine/Plugins/Experimental/ChaosEditor/Source/FractureEditor/Private/FractureToolCutter.cpp
+++ b/Engine/Plugins/Experimental/ChaosEditor/Source/FractureEditor/Private/FractureToolCutter.cpp
@@ -16,6 +16,8 @@
 #include "GeometryCollection/GeometryCollectionAlgo.h"
 #include "Components/DynamicMeshComponent.h"
 
+#include "FractureEngineMaterials.h"
+
 #include "Voronoi/Voronoi.h"
 #include "PlanarCut.h"
 #include "DynamicMesh/MeshTransforms.h"
@@ -242,6 +244,20 @@ void UFractureToolCutterBase::UpdateDefaultRandomSeed()
 }
 
 
+void UFractureToolCutterBase::PostFractureProcess(const FFractureToolContext& FractureContext, int32 FirstNewGeometryIndex)
+{
+	Super::PostFractureProcess(FractureContext, FirstNewGeometryIndex);
+
+	// Apply Internal Material ID to new fracture geometry.
+	int32 InternalMaterialID = CutterSettings->GetInternalMaterialID();
+	if (InternalMaterialID > INDEX_NONE)
+	{
+		FGeometryCollection& Collection = *FractureContext.GetGeometryCollection();
+		FFractureEngineMaterials::SetMaterialOnGeometryAfter(Collection, FirstNewGeometryIndex, FFractureEngineMaterials::ETargetFaces::InternalFaces, InternalMaterialID);
+		Collection.ReindexMaterials();
+	}
+}
+
 void FCellNoisePreviewOp::CalculateResult(FProgressCancel* Progress)
 {
 	Result->SetNum(ComputedDiagrams.Num());
diff --git a/Engine/Plugins/Experimental/ChaosEditor/Source/FractureEditor/Private/FractureToolCutter.h b/Engine/Plugins/Experimental/ChaosEditor/Source/FractureEditor/Private/FractureToolCutter.h
--- a/Engine/Plugins/Experimental/ChaosEditor/Source/FractureEditor/Private/FractureToolCutter.h
+++ b/Engine/Plugins/Experimental/ChaosEditor/Source/FractureEditor/Private/FractureToolCutter.h
@@ -43,6 +43,29 @@ public:
 
 	void TransferNoiseSettings(FNoiseSettings& NoiseSettingsOut);
 
+	/** Material to set for internal faces on fracture. 'Automatic' will use the most common 'internal' material in each geometry, or the last valid material if no internal faces are found */
+	UPROPERTY(EditAnywhere, Category = Materials, meta = (TransientToolProperty, DisplayName = "Internal Material", GetOptions = GetMaterialNamesFunc, NoResetToDefault))
+	FString InternalMaterial;
+
+	UFUNCTION()
+	const TArray<FString>& GetMaterialNamesFunc() { return ActiveMaterialNamesList; }
+
+	void UpdateActiveMaterialNames(TArray<FString> InMaterialNamesList)
+	{
+		ActiveMaterialNamesList = InMaterialNamesList;
+		if (!ActiveMaterialNamesList.Contains(InternalMaterial))
+		{
+			InternalMaterial = ActiveMaterialNamesList.IsEmpty() ? FString() : ActiveMaterialNamesList[0];
+		}
+	}
+
+	// Get the user-chosen internal material ID, or INDEX_NONE if 'Automatic' is chosen
+	int32 GetInternalMaterialID()
+	{
+		int32 NameIndex = ActiveMaterialNamesList.Find(InternalMaterial);
+		return FMath::Max((int32)INDEX_NONE, NameIndex - 1);
+	}
+
 	/** Random number generator seed for repeatability. If the value is -1, a different random seed will be used every time, otherwise the specified seed will always be used */
 	UPROPERTY(EditAnywhere, Category = CommonFracture, meta = (DisplayName = "Random Seed", UIMin = "-1", UIMax = "1000", ClampMin = "-1"))
 	int32 RandomSeed;
@@ -141,6 +164,11 @@ public:
 		}
 		return MaxDisp;
 	}
+
+private:
+
+	UPROPERTY(meta = (TransientToolProperty))
+	TArray<FString> ActiveMaterialNamesList;
 };
 
 /** Settings related to the collision properties of the fractured mesh pieces */
@@ -198,8 +226,22 @@ public:
 		UpdateDefaultRandomSeed();
 	}
 
+	virtual void Setup() override
+	{
+		Super::Setup();
+		CutterSettings->UpdateActiveMaterialNames(GetSelectedComponentMaterialNames(true));
+	}
+
+	virtual void SelectedBonesChanged() override
+	{
+		Super::SelectedBonesChanged();
+		CutterSettings->UpdateActiveMaterialNames(GetSelectedComponentMaterialNames(true));
+	}
+
 	FBox GetCombinedBounds(const TArray<FFractureToolContext>& Contexts) const;
 
+	// Post-processing; sets user-specified internal material on fracture (if not automatic)
+	virtual void PostFractureProcess(const FFractureToolContext& FractureContext, int32 FirstNewGeometryIndex) override;
 
 	virtual void UpdateUseGizmo(bool bUseGizmo)
 	{
diff --git a/Engine/Plugins/Experimental/ChaosEditor/Source/FractureEditor/Private/FractureToolGenerators.cpp b/Engine/Plugins/Experimental/ChaosEditor/Source/FractureEditor/Private/FractureToolGenerators.cpp
--- a/Engine/Plugins/Experimental/ChaosEditor/Source/FractureEditor/Private/FractureToolGenerators.cpp
+++ b/Engine/Plugins/Experimental/ChaosEditor/Source/FractureEditor/Private/FractureToolGenerators.cpp
@@ -88,7 +88,7 @@ public:
 	SLATE_ARGUMENT(FText, DefaultNameOverride)
 
 	SLATE_ARGUMENT(bool, bSplitIslands)
-	SLATE_ARGUMENT(bool, bAddInternalMaterials)
+	SLATE_ARGUMENT(bool, bReimportToMeshOutput)
 
 	/** Action to perform when create clicked */
 	SLATE_EVENT(FOnPathChosen, OnCreateAssetAction)
@@ -143,8 +143,8 @@ private:
 	/** Whether to split out each connected components into separate geometry */
 	bool bSplitIslands = false;
 
-	/** Whether to duplicate materials to create internal materials when creating from a static mesh */
-	bool bAddInternalMaterials = true;
+	/** Whether static mesh inputs were created by the 'ToMesh' tool in Fracture Mode */
+	bool bFromToMeshTool = false;
 
 	/** Filename textbox widget */
 	TSharedPtr<SEditableTextBox> FileNameWidget;
@@ -274,8 +274,8 @@ void SCreateGeometryCollectionFromObject::Construct(const FArguments& InArgs, TS
 					.VAlign(VAlign_Center)
 					[
 						SNew(STextBlock)
-						.Text(LOCTEXT("AddInternalsLabel", "Add Internal Materials"))
-						.ToolTipText(LOCTEXT("InternalMatsToolTip", "If checked, static mesh materials will be duplicated to create materials for internal surfaces. Does not apply if creating from a Geometry Collection source."))
+						.Text(LOCTEXT("ReimportToMeshLabel", "Reimport From 'ToMesh'"))
+						.ToolTipText(LOCTEXT("FromToMeshToolTip", "Only check for meshes that were created by the ToMesh tool in Fracture Mode. Does not apply if creating from a Geometry Collection source."))
 					]
 
 				+ SHorizontalBox::Slot()
@@ -284,7 +284,7 @@ void SCreateGeometryCollectionFromObject::Construct(const FArguments& InArgs, TS
 						SNew(SCheckBox)
 						.IsChecked_Lambda([this]()->ECheckBoxState
 							{
-								return bAddInternalMaterials ? ECheckBoxState::Checked : ECheckBoxState::Unchecked;
+								return bFromToMeshTool ? ECheckBoxState::Checked : ECheckBoxState::Unchecked;
 							})
 						.OnCheckStateChanged_Lambda([this](ECheckBoxState NewState)
 							{
@@ -292,9 +292,9 @@ void SCreateGeometryCollectionFromObject::Construct(const FArguments& InArgs, TS
 								{
 									return;
 								}
-								bAddInternalMaterials = NewState == ECheckBoxState::Checked;
+								bFromToMeshTool = NewState == ECheckBoxState::Checked;
 							})
-						.ToolTipText(LOCTEXT("InternalMatsToolTip", "If checked, static mesh materials will be duplicated to create materials for internal surfaces. Does not apply if creating from a Geometry Collection source."))
+						.ToolTipText(LOCTEXT("FromToMeshToolTip", "Only check for meshes that were created by the ToMesh tool in Fracture Mode. Does not apply if creating from a Geometry Collection source."))
 					]
 				]
 				+ SVerticalBox::Slot()
@@ -326,7 +326,7 @@ void SCreateGeometryCollectionFromObject::Construct(const FArguments& InArgs, TS
 							}
 							bSplitIslands = NewState == ECheckBoxState::Checked;
 						})
-						.ToolTipText(LOCTEXT("SplitComponentsToolTip", "If checked, triangles that are not topologically connected will be assigned separate bones on import."))
+						.ToolTipText(LOCTEXT("SplitComponentsToolTip", "If checked, triangles that are not topologically connected will be assigned separate bones on import. Does not apply if creating from a Geometry Collection source."))
 					]
 				]
 			]
@@ -375,7 +375,7 @@ void SCreateGeometryCollectionFromObject::RequestDestroyParentWindow()
 FReply SCreateGeometryCollectionFromObject::OnCreateAssetFromActorClicked()
 {
 	RequestDestroyParentWindow();
-	OnCreateAssetAction.ExecuteIfBound(AssetPath / FileNameWidget->GetText().ToString(), bAddInternalMaterials, bSplitIslands);
+	OnCreateAssetAction.ExecuteIfBound(AssetPath / FileNameWidget->GetText().ToString(), bFromToMeshTool, bSplitIslands);
 	return FReply::Handled();
 }
 
@@ -529,7 +529,7 @@ void UFractureToolGenerateAsset::OpenGenerateAssetDialog(TArray<AActor*>& Actors
 
 }
 
-void UFractureToolGenerateAsset::OnGenerateAssetPathChosen(const FString& InAssetPath, bool bAddInternalMaterials, bool bSplitComponents, TArray<AActor*> Actors)
+void UFractureToolGenerateAsset::OnGenerateAssetPathChosen(const FString& InAssetPath, bool bFromToMeshTool, bool bSplitComponents, TArray<AActor*> Actors)
 {	
 	FScopedTransaction Transaction(LOCTEXT("GenerateAsset", "Generate Geometry Collection Asset"));
 
@@ -548,7 +548,7 @@ void UFractureToolGenerateAsset::OnGenerateAssetPathChosen(const FString& InAsse
 
 		AGeometryCollectionActor* GeometryCollectionActor = Cast<AGeometryCollectionActor>(FirstActor);
 		
-		GeometryCollectionActor = ConvertActorsToGeometryCollection(InAssetPath, bAddInternalMaterials, bSplitComponents, Actors);
+		GeometryCollectionActor = ConvertActorsToGeometryCollection(InAssetPath, false/*bAddInternalMaterials*/, bSplitComponents, Actors, bFromToMeshTool);
 
 		GeometryCollectionComponent = GeometryCollectionActor->GetGeometryCollectionComponent();
 
@@ -590,8 +590,9 @@ void UFractureToolGenerateAsset::OnGenerateAssetPathChosen(const FString& InAsse
 	}
 }
 
-AGeometryCollectionActor* UFractureToolGenerateAsset::ConvertActorsToGeometryCollection(const FString& InAssetPath, bool bAddInternalMaterials, bool bSplitComponents, TArray<AActor*>& Actors)
+AGeometryCollectionActor* UFractureToolGenerateAsset::ConvertActorsToGeometryCollection(const FString& InAssetPath, bool bAddInternalMaterials, bool bSplitComponents, TArray<AActor*>& Actors, bool bFromToMeshTool)
 {
+	ensure(!bAddInternalMaterials); // we should not use the 'add internal materials' path anymore, as we move away from the odd-numbered-internal-material convention
 	ensure(Actors.Num() > 0);
 	AActor* FirstActor = Actors[0];
 	const FString& Name = FirstActor->GetActorLabel();
@@ -630,9 +631,11 @@ AGeometryCollectionActor* UFractureToolGenerateAsset::ConvertActorsToGeometryCol
 				// Record the contributing source on the asset.
 				FSoftObjectPath SourceSoftObjectPath(ComponentStaticMesh);
 				decltype(FGeometryCollectionSource::SourceMaterial) SourceMaterials(StaticMeshComponent->GetMaterials());
-				FracturedGeometryCollection->GeometrySource.Add({ SourceSoftObjectPath, ComponentTransform, SourceMaterials, bAddInternalMaterials, bSplitComponents });
-
-				FGeometryCollectionEngineConversion::AppendStaticMesh(ComponentStaticMesh, SourceMaterials, ComponentTransform, FracturedGeometryCollection, false, bAddInternalMaterials, bSplitComponents);
+				bool bSetInternalFromMaterialIndex = bFromToMeshTool;
+				
+				FracturedGeometryCollection->GeometrySource.Emplace(SourceSoftObjectPath, ComponentTransform, SourceMaterials, bSplitComponents, bSetInternalFromMaterialIndex);
+				
+				FGeometryCollectionEngineConversion::AppendStaticMesh(ComponentStaticMesh, SourceMaterials, ComponentTransform, FracturedGeometryCollection, false/*bReindexMaterials*/, bAddInternalMaterials, bSplitComponents, bSetInternalFromMaterialIndex);
 			}
 		}
 
@@ -664,15 +667,16 @@ AGeometryCollectionActor* UFractureToolGenerateAsset::ConvertActorsToGeometryCol
 				{
 					SourceMaterials[MaterialIndex] = GeometryCollectionComponent->GetMaterial(MaterialIndex);
 				}
-				FracturedGeometryCollection->GeometrySource.Add({ SourceSoftObjectPath, ComponentTransform, SourceMaterials });
+				constexpr bool bSplitCollectionComponents = false, bSetInternalFromMaterialIndex = false;
+				FracturedGeometryCollection->GeometrySource.Emplace(SourceSoftObjectPath, ComponentTransform, SourceMaterials, bSplitCollectionComponents, bSetInternalFromMaterialIndex);
 
-				FGeometryCollectionEngineConversion::AppendGeometryCollection(RestCollection, GeometryCollectionComponent, ComponentTransform, FracturedGeometryCollection, false);
+				FGeometryCollectionEngineConversion::AppendGeometryCollection(RestCollection, GeometryCollectionComponent, ComponentTransform, FracturedGeometryCollection, false /*bReindexMaterials*/);
 
 			}
 		}
 	}
 
-	FracturedGeometryCollection->InitializeMaterials();
+	FracturedGeometryCollection->InitializeMaterials(bAddInternalMaterials);
 
 	AddSingleRootNodeIfRequired(FracturedGeometryCollection);
 
@@ -845,7 +849,8 @@ void UFractureToolResetAsset::Execute(TWeakPtr<FFractureEditorModeToolkit> InToo
 					const UObject* SourceMesh = Source.SourceGeometryObject.TryLoad();
 					if (const UStaticMesh* SourceStaticMesh = Cast<UStaticMesh>(SourceMesh))
 					{
-						FGeometryCollectionEngineConversion::AppendStaticMesh(SourceStaticMesh, Source.SourceMaterial, Source.LocalTransform, GeometryCollectionObject, false, Source.bAddInternalMaterials, Source.bSplitComponents);
+						bool bLegacyAddInternal = Source.bAddInternalMaterials;
+						FGeometryCollectionEngineConversion::AppendStaticMesh(SourceStaticMesh, Source.SourceMaterial, Source.LocalTransform, GeometryCollectionObject, false, bLegacyAddInternal, Source.bSplitComponents, Source.bSetInternalFromMaterialIndex);
 					}
 					else if (const USkeletalMesh* SourceSkeletalMesh = Cast<USkeletalMesh>(SourceMesh))
 					{
@@ -858,21 +863,18 @@ void UFractureToolResetAsset::Execute(TWeakPtr<FFractureEditorModeToolkit> InToo
 					}
 				}
 
-				GeometryCollectionObject->InitializeMaterials();
+				constexpr bool bHasInternalMaterials = false;
+				GeometryCollectionObject->InitializeMaterials(bHasInternalMaterials);
 
-				if (bKeepPreviousMaterials)
+				// attempt to keep previously-set materials (as long as the source doesn't have even more materials)
+				int32 NewMatNum = GeometryCollectionObject->Materials.Num(), OldMatNum = OldMaterials.Num();
+				if (bKeepPreviousMaterials && NewMatNum <= OldMatNum)
 				{
-					int32 NewMatNum = GeometryCollectionObject->Materials.Num(), OldMatNum = OldMaterials.Num();
-					// if the source asset was changed, number of materials might have changed; only copy to the extent the two arrays match
-					int32 NumToCopy = FMath::Min(NewMatNum, OldMatNum);
-					for (int32 MatIdx = 0; MatIdx + 1 < NumToCopy; MatIdx++)
+					GeometryCollectionObject->Materials.SetNum(OldMatNum);
+					for (int32 MatIdx = 0; MatIdx < OldMatNum; MatIdx++)
 					{
 						GeometryCollectionObject->Materials[MatIdx] = OldMaterials[MatIdx];
 					}
-					if (NumToCopy > 0) // copy the selection material
-					{
-						GeometryCollectionObject->Materials[NewMatNum - 1] = OldMaterials[OldMatNum - 1];
-					}
 				}
 
 				const UFractureModeSettings* ModeSettings = GetDefault<UFractureModeSettings>();
diff --git a/Engine/Plugins/Experimental/ChaosEditor/Source/FractureEditor/Private/FractureToolGenerators.h b/Engine/Plugins/Experimental/ChaosEditor/Source/FractureEditor/Private/FractureToolGenerators.h
--- a/Engine/Plugins/Experimental/ChaosEditor/Source/FractureEditor/Private/FractureToolGenerators.h
+++ b/Engine/Plugins/Experimental/ChaosEditor/Source/FractureEditor/Private/FractureToolGenerators.h
@@ -30,8 +30,8 @@ public:
 
 private:
 	void OpenGenerateAssetDialog(TArray<AActor*>& Actors);
-	void OnGenerateAssetPathChosen(const FString& InAssetPath, bool bAddInternalMaterials, bool bSplitComponents, TArray<AActor*> Actors);
-	AGeometryCollectionActor* ConvertActorsToGeometryCollection(const FString& InAssetPath, bool bAddInternalMaterials, bool bSplitComponents, TArray<AActor*>& Actors);
+	void OnGenerateAssetPathChosen(const FString& InAssetPath, bool bFromToMeshTool, bool bSplitComponents, TArray<AActor*> Actors);
+	AGeometryCollectionActor* ConvertActorsToGeometryCollection(const FString& InAssetPath, bool bAddInternalMaterials, bool bSplitComponents, TArray<AActor*>& Actors, bool bFromToMeshTool);
 	class AGeometryCollectionActor* CreateNewGeometryActor(const FString& Name, const FTransform& Transform, bool AddMaterials /*= false*/);
 	static ULevel* GetSelectedLevel();
 	static AActor* AddActor(ULevel* InLevel, UClass* Class);
diff --git a/Engine/Plugins/Experimental/ChaosEditor/Source/FractureEditor/Private/FractureToolMaterials.cpp b/Engine/Plugins/Experimental/ChaosEditor/Source/FractureEditor/Private/FractureToolMaterials.cpp
new file mode 100644
--- /dev/null
+++ b/Engine/Plugins/Experimental/ChaosEditor/Source/FractureEditor/Private/FractureToolMaterials.cpp
@@ -0,0 +1,184 @@
+// Copyright Epic Games, Inc. All Rights Reserved.
+
+#include "FractureToolMaterials.h"
+
+#include "FractureToolContext.h"
+
+#include "GeometryCollection/GeometryCollectionObject.h"
+#include "GeometryCollection/GeometryCollectionComponent.h"
+
+#include "FractureEngineMaterials.h"
+
+#include UE_INLINE_GENERATED_CPP_BY_NAME(FractureToolMaterials)
+
+#define LOCTEXT_NAMESPACE "FractureToolMaterials"
+
+void UFractureMaterialsSettings::AddMaterialSlot()
+{
+	UFractureToolMaterials* MaterialsTool = Cast<UFractureToolMaterials>(OwnerTool.Get());
+	MaterialsTool->AddMaterialSlot();
+}
+
+void UFractureMaterialsSettings::RemoveMaterialSlot()
+{
+	UFractureToolMaterials* MaterialsTool = Cast<UFractureToolMaterials>(OwnerTool.Get());
+	MaterialsTool->RemoveMaterialSlot();
+}
+
+void UFractureToolMaterials::RemoveMaterialSlot()
+{
+	TSet<UGeometryCollectionComponent*> GeomCompSelection;
+	GetSelectedGeometryCollectionComponents(GeomCompSelection);
+	FScopedTransaction Transaction(LOCTEXT("RemoveMaterialSlot", "Remove Material from Geometry Collection(s)"), !GeomCompSelection.IsEmpty());
+	for (UGeometryCollectionComponent* GeometryCollectionComponent : GeomCompSelection)
+	{
+		FGeometryCollectionEdit Edit(GeometryCollectionComponent, GeometryCollection::EEditUpdate::Rest);
+		UGeometryCollection* Collection = Edit.GetRestCollection();
+		if (Collection->RemoveLastMaterialSlot())
+		{
+			GeometryCollectionComponent->MarkRenderDynamicDataDirty();
+			GeometryCollectionComponent->MarkRenderStateDirty();
+		}
+	}
+	UpdateActiveMaterialsList();
+}
+
+void UFractureToolMaterials::AddMaterialSlot()
+{
+	TSet<UGeometryCollectionComponent*> GeomCompSelection;
+	GetSelectedGeometryCollectionComponents(GeomCompSelection);
+	FScopedTransaction Transaction(LOCTEXT("AddMaterialSlot", "Add Material to Geometry Collection(s)"), !GeomCompSelection.IsEmpty());
+	for (UGeometryCollectionComponent* GeometryCollectionComponent : GeomCompSelection)
+	{
+		FGeometryCollectionEdit Edit(GeometryCollectionComponent, GeometryCollection::EEditUpdate::Rest);
+		UGeometryCollection* Collection = Edit.GetRestCollection();
+
+		int32 NewSlotIdx = Collection->AddNewMaterialSlot();
+		if (NewSlotIdx > 0)
+		{
+			// copy an adjacent material into the new slot on the component as well
+			GeometryCollectionComponent->SetMaterial(NewSlotIdx, GeometryCollectionComponent->GetMaterial(NewSlotIdx - 1));
+		}
+
+		GeometryCollectionComponent->MarkRenderDynamicDataDirty();
+		GeometryCollectionComponent->MarkRenderStateDirty();
+	}
+	UpdateActiveMaterialsList();
+}
+
+UFractureToolMaterials::UFractureToolMaterials(const FObjectInitializer& ObjInit)
+	: Super(ObjInit)
+{
+	MaterialsSettings = NewObject<UFractureMaterialsSettings>(GetTransientPackage(), UFractureMaterialsSettings::StaticClass());
+	MaterialsSettings->OwnerTool = this;
+}
+
+bool UFractureToolMaterials::CanExecute() const
+{
+	if (!IsGeometryCollectionSelected())
+	{
+		return false;
+	}
+
+	return true;
+}
+
+FText UFractureToolMaterials::GetDisplayText() const
+{
+	return FText(LOCTEXT("FractureToolMaterials", "Edit geometry collection materials and default material assignments for new faces"));
+}
+
+FText UFractureToolMaterials::GetTooltipText() const
+{
+	return FText(LOCTEXT("FractureToolMaterialsTooltip", "Allows direct editing of materials on a geometry collection, as well as editing of the default handling."));
+}
+
+FSlateIcon UFractureToolMaterials::GetToolIcon() const
+{
+	return FSlateIcon("FractureEditorStyle", "FractureEditor.ToMesh");
+}
+
+void UFractureToolMaterials::RegisterUICommand(FFractureEditorCommands* BindingContext)
+{
+	UI_COMMAND_EXT(BindingContext, UICommandInfo, "Material", "Material", "Update geometry materials, especially for new internal faces resulting from fracture.", EUserInterfaceActionType::ToggleButton, FInputChord());
+	BindingContext->Materials = UICommandInfo;
+}
+
+TArray<UObject*> UFractureToolMaterials::GetSettingsObjects() const
+{
+	TArray<UObject*> Settings;
+	Settings.Add(MaterialsSettings);
+	return Settings;
+}
+
+void UFractureToolMaterials::FractureContextChanged()
+{
+}
+
+void UFractureToolMaterials::Render(const FSceneView* View, FViewport* Viewport, FPrimitiveDrawInterface* PDI)
+{
+}
+
+void UFractureToolMaterials::PostEditChangeChainProperty(struct FPropertyChangedChainEvent& PropertyChangedEvent)
+{
+	// update any cached data 
+}
+
+TArray<FFractureToolContext> UFractureToolMaterials::GetFractureToolContexts() const
+{
+	TArray<FFractureToolContext> Contexts;
+
+	// A context is gathered for each selected GeometryCollection component
+	TSet<UGeometryCollectionComponent*> GeomCompSelection;
+	GetSelectedGeometryCollectionComponents(GeomCompSelection);
+
+	for (UGeometryCollectionComponent* GeometryCollectionComponent : GeomCompSelection)
+	{
+		FFractureToolContext FullSelection(GeometryCollectionComponent);
+		FullSelection.ConvertSelectionToRigidNodes();
+		Contexts.Add(FullSelection);
+		// TODO: consider also visualizing which faces will be updated
+	}
+
+	return Contexts;
+}
+
+
+int32 UFractureToolMaterials::ExecuteFracture(const FFractureToolContext& FractureContext)
+{
+	if (FractureContext.IsValid())
+	{
+		FGeometryCollection& Collection = *FractureContext.GetGeometryCollection();
+
+		
+
+		// Get ID to set and make sure it's valid
+		int32 MatID = MaterialsSettings->GetAssignMaterialID();
+		if (MatID == INDEX_NONE || MatID >= FractureContext.GetGeometryCollectionComponent()->GetNumMaterials())
+		{
+			return INDEX_NONE;
+		}
+
+		// convert enum to matching fracture engine materials enum
+		FFractureEngineMaterials::ETargetFaces TargetFaces =
+			(MaterialsSettings->ToFaces == EMaterialAssignmentTargets::AllFaces) ? FFractureEngineMaterials::ETargetFaces::AllFaces :
+			(MaterialsSettings->ToFaces == EMaterialAssignmentTargets::OnlyInternalFaces) ? FFractureEngineMaterials::ETargetFaces::InternalFaces :
+			FFractureEngineMaterials::ETargetFaces::ExternalFaces;
+
+		if (MaterialsSettings->bOnlySelected)
+		{
+			FFractureEngineMaterials::SetMaterial(Collection, FractureContext.GetSelection(), TargetFaces, MatID);
+		}
+		else
+		{
+			FFractureEngineMaterials::SetMaterialOnAllGeometry(Collection, TargetFaces, MatID);
+		}
+
+		Collection.ReindexMaterials();
+	}
+
+	return INDEX_NONE;
+}
+
+#undef LOCTEXT_NAMESPACE
+
diff --git a/Engine/Plugins/Experimental/ChaosEditor/Source/FractureEditor/Private/FractureToolMaterials.h b/Engine/Plugins/Experimental/ChaosEditor/Source/FractureEditor/Private/FractureToolMaterials.h
new file mode 100644
--- /dev/null
+++ b/Engine/Plugins/Experimental/ChaosEditor/Source/FractureEditor/Private/FractureToolMaterials.h
@@ -0,0 +1,134 @@
+// Copyright Epic Games, Inc. All Rights Reserved.
+
+#pragma once
+
+#include "FractureTool.h"
+
+#include "FractureToolMaterials.generated.h"
+
+class FFractureToolContext;
+struct FMeshDescription;
+
+UENUM()
+enum class EMaterialAssignmentTargets
+{
+	OnlyInternalFaces,
+	OnlyExternalFaces,
+	AllFaces
+};
+
+/** Settings related to geometry collection -> static mesh conversion **/
+UCLASS(config = EditorPerProjectUserSettings)
+class UFractureMaterialsSettings : public UFractureToolSettings
+{
+public:
+
+	GENERATED_BODY()
+
+	UFractureMaterialsSettings(const FObjectInitializer& ObjInit)
+		: Super(ObjInit)
+	{}
+
+	/** Add a new material slot to the selected geometry collections */
+	UFUNCTION(CallInEditor, Category = MaterialEditing, meta = (DisplayName = "Add Material Slot"))
+	void AddMaterialSlot();
+
+	/** Remove the last material slot from the selected geometry collections. (Will not remove the final material.) */
+	UFUNCTION(CallInEditor, Category = MaterialEditing, meta = (DisplayName = "Remove Material Slot"))
+	void RemoveMaterialSlot();
+
+	/** Material to assign to selected faces */
+	UPROPERTY(EditAnywhere, Category = MaterialEditing, meta = (TransientToolProperty, DisplayName = "Assign Material", GetOptions = GetMaterialNamesFunc, NoResetToDefault))
+	FString AssignMaterial;
+
+	/** Which subset of faces to update materials assignments on, for the selected geometry */
+	UPROPERTY(EditAnywhere, Category = MaterialEditing)
+	EMaterialAssignmentTargets ToFaces = EMaterialAssignmentTargets::OnlyInternalFaces;
+
+	/** Whether to only assign materials for faces in the selected bones, or the whole geometry collection */
+	UPROPERTY(EditAnywhere, Category = MaterialEditing)
+	bool bOnlySelected = true;
+
+	UFUNCTION()
+	const TArray<FString>& GetMaterialNamesFunc() { return AssignMaterialNamesList; }
+
+	void UpdateActiveMaterialNames(TArray<FString> InMaterialNamesList)
+	{
+		AssignMaterialNamesList = InMaterialNamesList;
+		if (!AssignMaterialNamesList.Contains(AssignMaterial))
+		{
+			AssignMaterial = AssignMaterialNamesList.IsEmpty() ? FString() : AssignMaterialNamesList[0];
+		}
+	}
+
+	int32 GetAssignMaterialID()
+	{
+		return AssignMaterialNamesList.Find(AssignMaterial);
+	}
+
+private:
+
+	UPROPERTY(meta = (TransientToolProperty))
+	TArray<FString> AssignMaterialNamesList;
+
+};
+
+
+UCLASS(DisplayName = "Materials Tool", Category = "FractureTools")
+class UFractureToolMaterials : public UFractureModalTool
+{
+public:
+	GENERATED_BODY()
+
+	UFractureToolMaterials(const FObjectInitializer& ObjInit);
+
+	///
+	/// UFractureModalTool Interface
+	///
+
+	/** This is the Text that will appear on the tool button to execute the tool **/
+	virtual FText GetDisplayText() const override;
+	virtual FText GetTooltipText() const override;
+	/** This is the Text that will appear on the button to execute the fracture **/
+	virtual FText GetApplyText() const override { return FText(NSLOCTEXT("FractureToolMaterials", "ExecuteMaterials", "Assign Materials")); }
+	virtual FSlateIcon GetToolIcon() const override;
+	virtual void RegisterUICommand(FFractureEditorCommands* BindingContext) override;
+	virtual TArray<UObject*> GetSettingsObjects() const override;
+	virtual void FractureContextChanged() override;
+	virtual void Render(const FSceneView* View, FViewport* Viewport, FPrimitiveDrawInterface* PDI) override;
+	/** Executes function that generates new geometry. Returns the first new geometry index. */
+	virtual int32 ExecuteFracture(const FFractureToolContext& FractureContext) override;
+	virtual bool CanExecute() const override;
+	virtual bool ExecuteUpdatesShape() const override
+	{
+		return false;
+	}
+	virtual void PostEditChangeChainProperty(struct FPropertyChangedChainEvent& PropertyChangedEvent) override;
+
+	virtual void SelectedBonesChanged()
+	{
+		Super::SelectedBonesChanged();
+		UpdateActiveMaterialsList();
+	}
+
+	// Called when the modal tool is entered
+	virtual void Setup()
+	{
+		Super::Setup();
+		UpdateActiveMaterialsList();
+	}
+
+	void UpdateActiveMaterialsList()
+	{
+		MaterialsSettings->UpdateActiveMaterialNames(GetSelectedComponentMaterialNames(false, false));
+	}
+
+	virtual TArray<FFractureToolContext> GetFractureToolContexts() const override;
+
+	void RemoveMaterialSlot();
+	void AddMaterialSlot();
+
+protected:
+	UPROPERTY(EditAnywhere, Category = Slicing)
+	TObjectPtr<UFractureMaterialsSettings> MaterialsSettings;
+};
diff --git a/Engine/Plugins/Experimental/ChaosEditor/Source/FractureEditor/Private/FractureToolRecomputeNormals.cpp b/Engine/Plugins/Experimental/ChaosEditor/Source/FractureEditor/Private/FractureToolRecomputeNormals.cpp
--- a/Engine/Plugins/Experimental/ChaosEditor/Source/FractureEditor/Private/FractureToolRecomputeNormals.cpp
+++ b/Engine/Plugins/Experimental/ChaosEditor/Source/FractureEditor/Private/FractureToolRecomputeNormals.cpp
@@ -106,7 +106,7 @@ void UFractureToolRecomputeNormals::FractureContextChanged()
 				{
 					for (int FIdx = FaceStart; FIdx < FaceStart + FaceCount; FIdx++)
 					{
-						if (Collection.Visible[FIdx] && (Collection.MaterialID[FIdx] % 2) == 1)
+						if (Collection.Visible[FIdx] && Collection.Internal[FIdx])
 						{
 							FIntVector Face = Collection.Indices[FIdx];
 							ShowVerts[Face.X - VertStart] = true;
diff --git a/Engine/Plugins/Experimental/ChaosEditor/Source/FractureEditor/Private/FractureToolUV.cpp b/Engine/Plugins/Experimental/ChaosEditor/Source/FractureEditor/Private/FractureToolUV.cpp
--- a/Engine/Plugins/Experimental/ChaosEditor/Source/FractureEditor/Private/FractureToolUV.cpp
+++ b/Engine/Plugins/Experimental/ChaosEditor/Source/FractureEditor/Private/FractureToolUV.cpp
@@ -30,6 +30,25 @@
 using namespace UE::Geometry;
 using namespace UE::Fracture;
 
+
+namespace
+{
+	UE::PlanarCut::ETargetFaces ConvertTargetFacesEnum(ETargetFaces TargetFaces)
+	{
+		return
+			TargetFaces == ETargetFaces::SelectedMaterialIDs ? UE::PlanarCut::ETargetFaces::CustomFaces :
+			TargetFaces == ETargetFaces::AllFaces ? UE::PlanarCut::ETargetFaces::AllFaces :
+			TargetFaces == ETargetFaces::ExternalFaces ? UE::PlanarCut::ETargetFaces::ExternalFaces :
+			UE::PlanarCut::ETargetFaces::InternalFaces;
+	}
+
+	bool TargetsMaterialIDs(ETargetFaces TargetFaces)
+	{
+		return TargetFaces != ETargetFaces::InternalFaces && TargetFaces != ETargetFaces::ExternalFaces && TargetFaces != ETargetFaces::AllFaces;
+	}
+}
+
+
 void UFractureAutoUVSettings::SetNumUVChannels(int32 NumUVChannels)
 {
 	NumUVChannels = FMath::Clamp<int32>(NumUVChannels, 1, GeometryCollectionUV::MAX_NUM_UV_CHANNELS);
@@ -206,9 +225,7 @@ void UFractureToolAutoUV::BoxProjectUVs()
 	int32 UVLayer = AutoUVSettings->GetSelectedChannelIndex();
 
 	TArray<int32> EmptyMaterialIDs;
-	UE::PlanarCut::EUseMaterials UseMaterialIDs =
-		AutoUVSettings->TargetMaterialIDs == ETargetMaterialIDs::SelectedIDs ? UE::PlanarCut::EUseMaterials::NoDefaultMaterials :
-		(AutoUVSettings->TargetMaterialIDs == ETargetMaterialIDs::AllIDs ? UE::PlanarCut::EUseMaterials::AllMaterials : UE::PlanarCut::EUseMaterials::OddMaterials);
+	UE::PlanarCut::ETargetFaces TargetFaces = ConvertTargetFacesEnum(AutoUVSettings->TargetFaces);
 
 	TSet<UGeometryCollectionComponent*> GeomCompSelection;
 	GetSelectedGeometryCollectionComponents(GeomCompSelection);
@@ -220,9 +237,9 @@ void UFractureToolAutoUV::BoxProjectUVs()
 		UE::PlanarCut::BoxProjectUVs(UVLayer,
 			*GeometryCollectionComponent->GetRestCollection()->GetGeometryCollection(),
 			(FVector3d)AutoUVSettings->ProjectionScale,
-			UseMaterialIDs,
-			AutoUVSettings->TargetMaterialIDs == ETargetMaterialIDs::OddIDs ? EmptyMaterialIDs : AutoUVSettings->MaterialIDs,
-			(FVector2f)AutoUVSettings->ProjectionUVOffset,
+			TargetFaces,
+			!TargetsMaterialIDs(AutoUVSettings->TargetFaces) ? EmptyMaterialIDs : AutoUVSettings->MaterialIDs)
+			(FVector2f)AutoU			(FVector2f)AutoUVSettings->ProjectionUVOffset,
 			AutoUVSettings->bAutoFitToBounds, AutoUVSettings->bCenterAtPivot, AutoUVSettings->bUniformProjectionScale);
 
 		GeometryCollectionComponent->MarkRenderDynamicDataDirty();
@@ -399,15 +416,6 @@ bool UFractureToolAutoUV::SaveGeneratedTexture(UE::Geometry::TImageBuilder<FVect
 	return true;
 }
 
-namespace
-{
-	UE::PlanarCut::EUseMaterials GetUseMaterials(ETargetMaterialIDs TargetIDs)
-	{
-		return TargetIDs == ETargetMaterialIDs::SelectedIDs ? UE::PlanarCut::EUseMaterials::NoDefaultMaterials :
-			(TargetIDs == ETargetMaterialIDs::AllIDs ? UE::PlanarCut::EUseMaterials::AllMaterials : UE::PlanarCut::EUseMaterials::OddMaterials);
-	}
-}
-
  class FLayoutUVOp : public FGeometryCollectionOperator
  {
  public:
@@ -419,7 +427,7 @@ namespace
 	int32 UVLayer;
 	int32 OutputRes;
 	int32 GutterSize;
-	UE::PlanarCut::EUseMaterials UseMaterialIDs;
+	UE::PlanarCut::ETargetFaces TargetFaces;
 	TArray<int32> MaterialIDs;
 
  	// TGenericDataOperator interface:
@@ -427,7 +435,7 @@ namespace
  	{
 		bool bLayoutSuccess =
 			UE::PlanarCut::UVLayout(UVLayer, *CollectionCopy, OutputRes, GutterSize,
-				UseMaterialIDs, MaterialIDs, true, Progress);
+				TargetFaces, MaterialIDs, true, Progress);
 		
 		if (bLayoutSuccess)
 		{
@@ -463,12 +471,12 @@ bool UFractureToolAutoUV::LayoutUVsForComponent(UGeometryCollectionComponent* Co
 	TUniquePtr<FLayoutUVOp> LayoutOp = MakeUnique<FLayoutUVOp>(Collection);
 	LayoutOp->OutputRes = (int32)AutoUVSettings->Resolution;
 	TArray<int32> EmptyMaterialIDs;
-	LayoutOp->MaterialIDs = AutoUVSettings->TargetMaterialIDs == ETargetMaterialIDs::OddIDs ? EmptyMaterialIDs : AutoUVSettings->MaterialIDs;
+	LayoutOp->MaterialIDs = !TargetsMaterialIDs(AutoUVSettings->TargetFaces) ? EmptyMaterialIDs : AutoUVSettings->MaterialIDs;
 
 	LayoutOp->OutputRes = (int32)AutoUVSettings->Resolution;
 	LayoutOp->UVLayer = AutoUVSettings->GetSelectedChannelIndex();
 	LayoutOp->GutterSize = AutoUVSettings->GutterSize;
-	LayoutOp->UseMaterialIDs = GetUseMaterials(AutoUVSettings->TargetMaterialIDs);
+	LayoutOp->TargetFaces = ConvertTargetFacesEnum(AutoUVSettings->TargetFaces);
 
 	int Result = RunCancellableGeometryCollectionOp<FLayoutUVOp>(Collection,
 		MoveTemp(LayoutOp), LOCTEXT("ComputingUVLayoutMessage", "Computing UV Layout"));
@@ -489,15 +497,15 @@ public:
 	FIndex4i Attributes;
 	UE::PlanarCut::FTextureAttributeSettings AttribSettings;
 	TUniquePtr<UE::Geometry::TImageBuilder<FVector4f>> ImageBuilder;
-	UE::PlanarCut::EUseMaterials UseMaterialIDs;
+	UE::PlanarCut::ETargetFaces TargetFaces;
 	TArray<int32> MaterialIDs;
 
 	// TGenericDataOperator interface:
 	virtual void CalculateResult(FProgressCancel* Progress) override
 	{
-		bool bBakeSuccess = UE::PlanarCut::TextureInternalSurfaces(
+		bool bBakeSuccess = UE::PlanarCut::TextureSpecifiedFaces(
 			UVLayer, *CollectionCopy, GutterSize, Attributes, AttribSettings, *ImageBuilder,
-			UseMaterialIDs, MaterialIDs, Progress);
+			TargetFaces, MaterialIDs, Progress);
 
 		if (!bBakeSuccess || (Progress && Progress->Cancelled()))
 		{
@@ -530,11 +538,11 @@ void UFractureToolAutoUV::BakeTextureForComponent(UGeometryCollectionComponent*
 
 	TUniquePtr<FBakeTextureOp> BakeOp = MakeUnique<FBakeTextureOp>(Collection);
 	TArray<int32> EmptyMaterialIDs;
-	BakeOp->MaterialIDs = AutoUVSettings->TargetMaterialIDs == ETargetMaterialIDs::OddIDs ? EmptyMaterialIDs : AutoUVSettings->MaterialIDs;
+	BakeOp->MaterialIDs = !TargetsMaterialIDs(AutoUVSettings->TargetFaces) ? EmptyMaterialIDs : AutoUVSettings->MaterialIDs;
 
 	BakeOp->UVLayer = AutoUVSettings->GetSelectedChannelIndex();
 	BakeOp->GutterSize = AutoUVSettings->GutterSize;
-	BakeOp->UseMaterialIDs = GetUseMaterials(AutoUVSettings->TargetMaterialIDs);
+	BakeOp->TargetFaces = ConvertTargetFacesEnum(AutoUVSettings->TargetFaces);
 
 	int32 OutputRes = (int32)AutoUVSettings->Resolution;
 	FImageDimensions Dimensions(OutputRes, OutputRes);
diff --git a/Engine/Plugins/Experimental/ChaosEditor/Source/FractureEditor/Private/FractureToolUV.h b/Engine/Plugins/Experimental/ChaosEditor/Source/FractureEditor/Private/FractureToolUV.h
--- a/Engine/Plugins/Experimental/ChaosEditor/Source/FractureEditor/Private/FractureToolUV.h
+++ b/Engine/Plugins/Experimental/ChaosEditor/Source/FractureEditor/Private/FractureToolUV.h
@@ -38,16 +38,18 @@ enum class ETextureType
 };
 
 UENUM()
-enum class ETargetMaterialIDs
+enum class ETargetFaces
 {
-	// Target faces with odd-numbered material IDs (corresponding to internal fracture faces)
-	OddIDs,
-	// Target faces with material IDs matching a custom list, as well as those with odd-numbered IDs (corresponding to internal fracture faces)
-	OddAndSelectedIDs,
+	// Target internal faces
+	InternalFaces,
+	// Target internal faces and those with material IDs matching a custom list
+	InternalFacesAndSelectedMaterialIDs,
 	// Target faces with material IDs matching a custom list
-	SelectedIDs,
+	SelectedMaterialIDs,
+	// Target external faces
+	ExternalFaces,
 	// Target all faces
-	AllIDs
+	AllFaces
 };
 
 /** Settings for UV layout and texture baking on the geometry collection **/
@@ -132,12 +134,12 @@ public:
 	UFUNCTION(CallInEditor, Category = Unwrap, meta = (DisplayName = "Box Project UVs"))
 	void BoxProjectUVs();
 
-	/** Choose whether to texture only faces with odd material IDs (corresponding to internal faces) or a custom selection */
+	/** Choose whether to target internal faces, or a custom selection */
 	UPROPERTY(EditAnywhere, Category = MapSettings)
-	ETargetMaterialIDs TargetMaterialIDs;
+	ETargetFaces TargetFaces;
 
 	/** Custom selection of material IDs to target for texturing */
-	UPROPERTY(EditAnywhere, Category = MapSettings, meta = (EditCondition = "TargetMaterialIDs != ETargetMaterialIDs::OddIDs && TargetMaterialIDs != ETargetMaterialIDs::AllIDs", EditConditionHides))
+	UPROPERTY(EditAnywhere, Category = MapSettings, meta = (EditCondition = "TargetFaces != ETargetFaces::InternalFaces && TargetFaces != ETargetFaces::AllFaces && TargetFaces != ETargetFaces::ExternalFaces", EditConditionHides))
 	TArray<int32> MaterialIDs;
 
 	/** The pixel resolution of the generated map */
diff --git a/Engine/Plugins/Experimental/ChaosEditor/Source/FractureEditor/Public/FractureTool.cpp b/Engine/Plugins/Experimental/ChaosEditor/Source/FractureEditor/Public/FractureTool.cpp
--- a/Engine/Plugins/Experimental/ChaosEditor/Source/FractureEditor/Public/FractureTool.cpp
+++ b/Engine/Plugins/Experimental/ChaosEditor/Source/FractureEditor/Public/FractureTool.cpp
@@ -27,6 +27,8 @@
 
 DEFINE_LOG_CATEGORY(LogFractureTool);
 
+#define LOCTEXT_NAMESPACE "FractureTool"
+
 void UFractureToolSettings::PostEditChangeProperty(FPropertyChangedEvent& PropertyChangedEvent)
 {
 	if (OwnerTool != nullptr)
@@ -175,6 +177,51 @@ void UFractureActionTool::GetSelectedGeometryCollectionComponents(TSet<UGeometry
 	}
 }
 
+TArray<FString> UFractureActionTool::GetSelectedComponentMaterialNames(bool bIncludeDefault, bool bUseFullNamesIfPossible)
+{
+	TArray<FString> MaterialNames;
+
+	if (bIncludeDefault)
+	{
+		MaterialNames.Add(LOCTEXT("AutomaticMaterialOption", "Automatic").ToString());
+	}
+	TSet<UGeometryCollectionComponent*> GeomCompSelection;
+	GetSelectedGeometryCollectionComponents(GeomCompSelection);
+	if (GeomCompSelection.Num() > 1 || !bUseFullNamesIfPossible)
+	{
+		int32 MinMaterials = -1;
+		for (UGeometryCollectionComponent* Component : GeomCompSelection)
+		{
+			int32 NumMaterials = Component->GetNumMaterials();
+			if (MinMaterials < 0 || NumMaterials < MinMaterials)
+			{
+				MinMaterials = NumMaterials;
+			}
+		}
+		for (int32 MatIdx = 0; MatIdx < MinMaterials; ++MatIdx)
+		{
+			MaterialNames.Add(FString::Printf(TEXT("[%d]"), MatIdx));
+		}
+	}
+	else if (GeomCompSelection.Num() == 1)
+	{
+		for (UGeometryCollectionComponent* Component : GeomCompSelection)
+		{
+			int32 NumMaterials = Component->GetNumMaterials();
+			for (int32 MatIdx = 0; MatIdx < NumMaterials; ++MatIdx)
+			{
+				if (MatIdx == Component->GetRestCollection()->GetBoneSelectedMaterialIndex())
+				{
+					continue;
+				}
+				MaterialNames.Add(FString::Printf(TEXT("[%d] %s"), MatIdx, *Component->GetMaterial(MatIdx)->GetName()));
+			}
+		}
+	}
+
+	return MaterialNames;
+}
+
 void UFractureActionTool::Refresh(FFractureToolContext& Context, FFractureEditorModeToolkit* Toolkit, bool bClearSelection)
 {
 	UGeometryCollectionComponent* GeometryCollectionComponent = Context.GetGeometryCollectionComponent();
@@ -316,6 +363,8 @@ void UFractureModalTool::Execute(TWeakPtr<FFractureEditorModeToolkit> InToolkit)
 			
 			if (FirstNewGeometryIndex > INDEX_NONE)
 			{
+				PostFractureProcess(FractureContext, FirstNewGeometryIndex);
+
 				FractureContext.GenerateGuids(FirstNewGeometryIndex);
 
 				// Based on the first new geometry index, generate a list of new transforms generated by the fracture.
@@ -420,3 +469,4 @@ FVector FVisualizationMappings::GetExplodedVector(int32 MappingIdx, const UGeome
 	return FVector::ZeroVector;
 }
 
+#undef LOCTEXT_NAMESPACE
diff --git a/Engine/Plugins/Experimental/ChaosEditor/Source/FractureEditor/Public/FractureTool.h b/Engine/Plugins/Experimental/ChaosEditor/Source/FractureEditor/Public/FractureTool.h
--- a/Engine/Plugins/Experimental/ChaosEditor/Source/FractureEditor/Public/FractureTool.h
+++ b/Engine/Plugins/Experimental/ChaosEditor/Source/FractureEditor/Public/FractureTool.h
@@ -79,6 +79,8 @@ public:
 
 public:
 	static void GetSelectedGeometryCollectionComponents(TSet<UGeometryCollectionComponent*>& GeomCompSelection);
+	static TArray<FString> GetSelectedComponentMaterialNames(bool bIncludeDefault, bool bUseFullNamesIfPossible = true);
+
 	
 protected:
 	static bool IsStaticMeshSelected();
@@ -159,6 +161,12 @@ public:
 	/** Executes function that generates new geometry. Returns the first new geometry index. */
 	virtual int32 ExecuteFracture(const FFractureToolContext& FractureContext) { return INDEX_NONE; }
 
+	// Optional processing after a successful fracture (i.e., will not be called if ExecuteFracture returns INDEX_NONE)
+	virtual void PostFractureProcess(const FFractureToolContext& FractureContext, int32 FirstNewGeometryIndex)
+	{
+
+	}
+
 	virtual void OnTick(float DeltaTime) {}
 
 	/** Draw callback from edmode*/
diff --git a/Engine/Plugins/Experimental/Fracture/Source/FractureEngine/Private/FractureEngineFracturing.cpp b/Engine/Plugins/Experimental/Fracture/Source/FractureEngine/Private/FractureEngineFracturing.cpp
--- a/Engine/Plugins/Experimental/Fracture/Source/FractureEngine/Private/FractureEngineFracturing.cpp
+++ b/Engine/Plugins/Experimental/Fracture/Source/FractureEngine/Private/FractureEngineFracturing.cpp
@@ -248,13 +248,6 @@ void FFractureEngineFracturing::VoronoiFracture(FManagedArrayCollection& InOutCo
 				
 				int ResultGeometryIndex = CutMultipleWithPlanarCells(VoronoiPlanarCells, *GeomCollection, TransformSelectionArr, InGrout, InCollisionSampleSpacing, InRandomSeed, FTransform().Identity);
 
-				{
-					// Add FacesGroup::Internal<bool> attribute to tag the newly created inside faces
-					const TManagedArray<int32>& MaterialIDs = GeomCollection->GetAttribute<int32>("MaterialID", FGeometryCollection::FacesGroup);
-
-					GeometryCollection::Facades::FCollectionMeshFacade::AddInternalAttribute(*GeomCollection, MaterialIDs.GetConstArray());
-				}
-
 				InOutCollection = (const FManagedArrayCollection&)(*GeomCollection);
 			}
 		}
@@ -322,13 +315,6 @@ void FFractureEngineFracturing::PlaneCutter(FManagedArrayCollection& InOutCollec
 
 		int ResultGeometryIndex = CutMultipleWithMultiplePlanes(CuttingPlanes, InternalSurfaceMaterials, *GeomCollection, TransformSelectionArr, GroutVal, CollisionSampleSpacingVal, InRandomSeed, FTransform().Identity);
 
-		{
-			// Add FacesGroup::Internal<bool> attribute to tag the newly created inside faces
-			const TManagedArray<int32>& MaterialIDs = GeomCollection->GetAttribute<int32>("MaterialID", FGeometryCollection::FacesGroup);
-
-			GeometryCollection::Facades::FCollectionMeshFacade::AddInternalAttribute(*GeomCollection, MaterialIDs.GetConstArray());
-		}
-
 		InOutCollection = (const FManagedArrayCollection&)(*GeomCollection);
 	}
 }
diff --git a/Engine/Plugins/Experimental/Fracture/Source/FractureEngine/Private/FractureEngineMaterials.cpp b/Engine/Plugins/Experimental/Fracture/Source/FractureEngine/Private/FractureEngineMaterials.cpp
new file mode 100644
--- /dev/null
+++ b/Engine/Plugins/Experimental/Fracture/Source/FractureEngine/Private/FractureEngineMaterials.cpp
@@ -0,0 +1,81 @@
+// Copyright Epic Games, Inc. All Rights Reserved.
+
+#include "FractureEngineMaterials.h"
+#include "GeometryCollection/Facades/CollectionMeshFacade.h"
+#include "GeometryCollection/Facades/CollectionTransformSelectionFacade.h"
+
+namespace UE::FractureEngineMaterialHelpers
+{
+	void SetMaterialOnFaceRange(const TManagedArray<bool>& Internal, TManagedArray<int32>& Material, int32 Start, int32 End, FFractureEngineMaterials::ETargetFaces TargetFaces, int32 MaterialID)
+	{
+		bool bAllFaces = TargetFaces == FFractureEngineMaterials::ETargetFaces::AllFaces;
+		bool bInternalFaces = TargetFaces == FFractureEngineMaterials::ETargetFaces::InternalFaces;
+		for (int32 FIdx = Start; FIdx < End; ++FIdx)
+		{
+			if (bAllFaces || bInternalFaces == Internal[FIdx])
+			{
+				Material[FIdx] = MaterialID;
+			}
+		}
+	}
+}
+ 
+
+void FFractureEngineMaterials::SetMaterial(FManagedArrayCollection& InCollection, const TArray<int32>& InBoneSelection, ETargetFaces TargetFaces, int32 MaterialID)
+{
+	GeometryCollection::Facades::FCollectionMeshFacade MeshFacade(InCollection);
+
+	if (!MeshFacade.IsValid())
+	{
+		return;
+	}
+
+	const TManagedArray<int32>& TransformToGeometryIndex = MeshFacade.TransformToGeometryIndexAttribute.Get();
+	const TManagedArray<int32>& FaceStart = MeshFacade.FaceStartAttribute.Get();
+	const TManagedArray<int32>& FaceCount = MeshFacade.FaceCountAttribute.Get();
+	const TManagedArray<bool>& Internal = MeshFacade.InternalAttribute.Get();
+	TManagedArray<int32>& Material = MeshFacade.MaterialIDAttribute.Modify();
+
+	GeometryCollection::Facades::FCollectionTransformSelectionFacade TransformSelectionFacade(InCollection);
+	TArray<int32> BoneIndices = InBoneSelection;
+	TransformSelectionFacade.ConvertSelectionToRigidNodes(BoneIndices);
+	TransformSelectionFacade.Sanitize(BoneIndices);
+
+	for (int32 Index : BoneIndices)
+	{
+		if (TransformToGeometryIndex[Index] > INDEX_NONE)
+		{
+			int32 Start = FaceStart[TransformToGeometryIndex[Index]];
+			int32 End = Start + FaceCount[TransformToGeometryIndex[Index]];
+			
+			UE::FractureEngineMaterialHelpers::SetMaterialOnFaceRange(Internal, Material, Start, End, TargetFaces, MaterialID);
+		}
+	}
+}
+
+
+void FFractureEngineMaterials::SetMaterialOnGeometryAfter(FManagedArrayCollection& InCollection, int32 FirstGeometryIndex, ETargetFaces TargetFaces, int32 MaterialID)
+{
+	GeometryCollection::Facades::FCollectionMeshFacade MeshFacade(InCollection);
+
+	if (!MeshFacade.IsValid())
+	{
+		return;
+	}
+
+	const TManagedArray<int32>& TransformToGeometryIndex = MeshFacade.TransformToGeometryIndexAttribute.Get();
+	const TManagedArray<int32>& FaceStart = MeshFacade.FaceStartAttribute.Get();
+	const TManagedArray<int32>& FaceCount = MeshFacade.FaceCountAttribute.Get();
+	const TManagedArray<bool>& Internal = MeshFacade.InternalAttribute.Get();
+	TManagedArray<int32>& Material = MeshFacade.MaterialIDAttribute.Modify();
+
+	for (int32 GeoIdx = FirstGeometryIndex, NumGeometry = FaceStart.Num(); GeoIdx < NumGeometry; ++GeoIdx)
+	{
+		int32 Start = FaceStart[GeoIdx];
+		int32 End = Start + FaceCount[GeoIdx];
+
+		UE::FractureEngineMaterialHelpers::SetMaterialOnFaceRange(Internal, Material, Start, End, TargetFaces, MaterialID);
+	}
+}
+
+
diff --git a/Engine/Plugins/Experimental/Fracture/Source/FractureEngine/Public/FractureEngineMaterials.h b/Engine/Plugins/Experimental/Fracture/Source/FractureEngine/Public/FractureEngineMaterials.h
new file mode 100644
--- /dev/null
+++ b/Engine/Plugins/Experimental/Fracture/Source/FractureEngine/Public/FractureEngineMaterials.h
@@ -0,0 +1,30 @@
+// Copyright Epic Games, Inc. All Rights Reserved.
+
+#pragma once
+
+#include "Containers/ContainersFwd.h" // IWYU pragma: keep
+
+class FGeometryCollection;
+struct FManagedArrayCollection;
+
+class FRACTUREENGINE_API FFractureEngineMaterials
+{
+public:
+
+	enum class ETargetFaces
+	{
+		InternalFaces,
+		ExternalFaces,
+		AllFaces
+	};
+
+	static void SetMaterial(FManagedArrayCollection& InCollection, const TArray<int32>& InBoneSelection, ETargetFaces TargetFaces, int32 MaterialID);
+
+	static void SetMaterialOnGeometryAfter(FManagedArrayCollection& InCollection, int32 FirstGeometryIndex, ETargetFaces TargetFaces, int32 MaterialID);
+
+	static void SetMaterialOnAllGeometry(FManagedArrayCollection& InCollection, ETargetFaces TargetFaces, int32 MaterialID)
+	{
+		SetMaterialOnGeometryAfter(InCollection, 0, TargetFaces, MaterialID);
+	}
+
+};
diff --git a/Engine/Plugins/Experimental/GeometryCollectionPlugin/Source/GeometryCollectionEditor/Private/GeometryCollection/GeometryCollectionFactory.cpp b/Engine/Plugins/Experimental/GeometryCollectionPlugin/Source/GeometryCollectionEditor/Private/GeometryCollection/GeometryCollectionFactory.cpp
--- a/Engine/Plugins/Experimental/GeometryCollectionPlugin/Source/GeometryCollectionEditor/Private/GeometryCollection/GeometryCollectionFactory.cpp
+++ b/Engine/Plugins/Experimental/GeometryCollectionPlugin/Source/GeometryCollectionEditor/Private/GeometryCollection/GeometryCollectionFactory.cpp
@@ -170,8 +170,8 @@ UObject* UGeometryCollectionFactory::FactoryCreateNew(UClass* Class, UObject* In
 		NewGeometryCollection->AppendGeometry(*GeometryCollectionData.Get<0>(), false, GeometryCollectionData.Get<2>());
 	}
 
-	// Add internal material and selection material
-	NewGeometryCollection->InitializeMaterials();
+	// De-duplicate materials and add selection material
+	NewGeometryCollection->InitializeMaterials(false);
 	GeometryCollectionAlgo::PrepareForSimulation(NewGeometryCollection->GetGeometryCollection().Get());
 
 	// Initial pivot : 
diff --git a/Engine/Plugins/Experimental/PlanarCutPlugin/Source/PlanarCut/Private/FractureAutoUV.cpp b/Engine/Plugins/Experimental/PlanarCutPlugin/Source/PlanarCut/Private/FractureAutoUV.cpp
--- a/Engine/Plugins/Experimental/PlanarCutPlugin/Source/PlanarCut/Private/FractureAutoUV.cpp
+++ b/Engine/Plugins/Experimental/PlanarCutPlugin/Source/PlanarCut/Private/FractureAutoUV.cpp
@@ -102,22 +102,24 @@ namespace UE
 } // namespace UE::UVPacking
 
 
-
 namespace UE { namespace PlanarCut {
 
 namespace {
-	inline bool IsTriActive(bool bIsVisible, int32 MaterialID, const TSet<int32>& InsideMaterials, bool bActivateInsideTriangles, EUseMaterials MaterialsPattern)
+	inline bool IsTriActive(bool bIsVisible, bool bIsInternal, int32 MaterialID, const TSet<int32>& InsideMaterials, bool bActivateInsideTriangles, ETargetFaces TargetFaces)
 	{
 		if (!bIsVisible) // never select invisible triangles
 		{
 			return false;
 		}
 
-		if (MaterialsPattern == EUseMaterials::AllMaterials && bActivateInsideTriangles)
+		if (TargetFaces == ETargetFaces::AllFaces && bActivateInsideTriangles)
 		{
 			return true;
 		}
-		if (MaterialsPattern == EUseMaterials::OddMaterials && ((MaterialID % 2) == 1) == bActivateInsideTriangles)
+		if (
+			(TargetFaces == ETargetFaces::InternalFaces && (bIsInternal == bActivateInsideTriangles)) ||
+			(TargetFaces == ETargetFaces::ExternalFaces && (bIsInternal != bActivateInsideTriangles))
+			)
 		{
 			return true;
 		}
@@ -135,7 +137,7 @@ namespace {
  * @return number of active triangles
  */
 int32 SetActiveTriangles(const FGeometryCollection* Collection, TArray<bool>& ActiveTrianglesOut,
-	bool bActivateInsideTriangles, EUseMaterials MaterialsPattern, TArrayView<int32> WhichMaterialsAreInside)
+	bool bActivateInsideTriangles, ETargetFaces TargetFaces, TArrayView<int32> WhichMaterialsAreInside)
 {
 	TSet<int32> InsideMaterials;
 	for (int32 MatID : WhichMaterialsAreInside)
@@ -147,7 +149,7 @@ int32 SetActiveTriangles(const FGeometryCollection* Collection, TArray<bool>& Ac
 	int32 NumTriangles = 0;
 	for (int TID = 0; TID < Collection->Indices.Num(); TID++)
 	{
-		bool bIsActive = IsTriActive(Collection->Visible[TID], Collection->MaterialID[TID], InsideMaterials, bActivateInsideTriangles, MaterialsPattern);
+		bool bIsActive = IsTriActive(Collection->Visible[TID], Collection->Internal[TID], Collection->MaterialID[TID], InsideMaterials, bActivateInsideTriangles, TargetFaces);
 		if (bIsActive)
 		{
 			ActiveTrianglesOut[TID] = true;
@@ -377,7 +379,7 @@ bool BoxProjectUVs(
 	int32 TargetUVLayer,
 	FGeometryCollection& Collection,
 	const FVector3d& BoxDimensions,
-	EUseMaterials MaterialsPattern,
+	ETargetFaces TargetFaces,
 	TArrayView<int32> WhichMaterials,
 	FVector2f OffsetUVs,
 	bool bOverrideBoxDimensionsWithBounds,
@@ -420,8 +422,8 @@ bool BoxProjectUVs(
 		for (int TID : Mesh.TriangleIndicesItr())
 		{
 			bool bIsTextureTri = IsTriActive(
-				AugmentedDynamicMesh::GetVisibility(Mesh, TID),
-				Mesh.Attributes()->GetMaterialID()->GetValue(TID), TargetMaterials, true, MaterialsPattern);
+				AugmentedDynamicMesh::GetVisibility(Mesh, TID), AugmentedDynamicMesh::GetInternal(Mesh, TID),
+				Mesh.Attributes()->GetMaterialID()->GetValue(TID), TargetMaterials, true, TargetFaces);
 			if (bIsTextureTri)
 			{
 				TargetTris.Add(TID);
@@ -471,7 +473,7 @@ bool UVLayout(
 	FGeometryCollection& Collection,
 	int32 UVRes,
 	float GutterSize,
-	EUseMaterials MaterialsPattern,
+	ETargetFaces TargetFaces,
 	TArrayView<int32> WhichMaterials,
 	bool bRecreateUVsForDegenerateIslands,
 	FProgressCancel* Progress
@@ -480,7 +482,7 @@ bool UVLayout(
 	FProgressCancel::FProgressScope InitScope = FProgressCancel::CreateScopeTo(Progress, .2);
 
 	TArray<bool> ActiveTriangles;
-	int32 NumActive = SetActiveTriangles(&Collection, ActiveTriangles, true, MaterialsPattern, WhichMaterials);
+	int32 NumActive = SetActiveTriangles(&Collection, ActiveTriangles, true, TargetFaces, WhichMaterials);
 	FGeomMesh UVMesh(TargetUVLayer, &Collection, ActiveTriangles, NumActive);
 
 	TArray<TArray<int32>> UVIslands;
@@ -614,19 +616,39 @@ bool UVLayout(
 
 
 bool TextureInternalSurfaces(
+	int32 TargetUVLayer,
+	FGeometryCollection& Collection,
+	int32 GutterSize,
+	UE::Geometry::FIndex4i BakeAttributes,
+	const FTextureAttributeSettings& AttributeSettings,
+	UE::Geometry::TImageBuilder<FVector4f>& TextureOut,
+	EUseMaterials MaterialsPattern,
+	TArrayView<int32> WhichMaterials,
+	FProgressCancel* Progress
+)
+{
+	ETargetFaces TargetFaces =
+		MaterialsPattern == EUseMaterials::AllMaterials ? ETargetFaces::AllFaces :
+		MaterialsPattern == EUseMaterials::OddMaterials ? ETargetFaces::InternalFaces : ETargetFaces::CustomFaces;
+	return TextureSpecifiedFaces(TargetUVLayer, Collection, GutterSize, BakeAttributes, AttributeSettings, TextureOut, TargetFaces, WhichMaterials, Progress);
+}
+
+
+
+bool TextureSpecifiedFaces(
 	int32 TargetUVLayer,
 	FGeometryCollection& Collection,
 	int32 GutterSize,
 	FIndex4i BakeAttributes,
 	const FTextureAttributeSettings& AttributeSettings,
 	TImageBuilder<FVector4f>& TextureOut,
-	EUseMaterials MaterialsPattern,
+	ETargetFaces TargetFaces,
 	TArrayView<int32> WhichMaterials,
 	FProgressCancel* Progress
 )
 {
 	TArray<bool> ToTextureTriangles;
-	int32 NumTextureTris = SetActiveTriangles(&Collection, ToTextureTriangles, true, MaterialsPattern, WhichMaterials);
+	int32 NumTextureTris = SetActiveTriangles(&Collection, ToTextureTriangles, true, TargetFaces, WhichMaterials);
 	FGeomFlatUVMesh UVMesh(TargetUVLayer, &Collection, ToTextureTriangles, NumTextureTris);
 
 	FImageOccupancyMap OccupancyMap;
@@ -636,9 +658,9 @@ bool TextureInternalSurfaces(
 
 	FGeomMesh ToTextureMesh(TargetUVLayer, &Collection, ToTextureTriangles, NumTextureTris);
 
-	// Outside triangles are identified by odd material ID
+	// Collect the external faces
 	TArray<bool> OutsideTriangles;
-	int32 NumOutsideTris = SetActiveTriangles(&Collection, OutsideTriangles, false, EUseMaterials::OddMaterials, {});
+	int32 NumOutsideTris = SetActiveTriangles(&Collection, OutsideTriangles, true, ETargetFaces::ExternalFaces, {});
 	FGeomMesh OutsideMesh(ToTextureMesh, OutsideTriangles, NumOutsideTris);
 	TMeshAABBTree3<FGeomMesh> OutsideSpatial(&OutsideMesh);
 
@@ -696,8 +718,8 @@ bool TextureInternalSurfaces(
 			for (int TID : Mesh.TriangleIndicesItr())
 			{
 				bool bIsTextureTri = IsTriActive(
-					AugmentedDynamicMesh::GetVisibility(Mesh, TID), 
-					Mesh.Attributes()->GetMaterialID()->GetValue(TID), TargetMaterials, true, MaterialsPattern);
+					AugmentedDynamicMesh::GetVisibility(Mesh, TID), AugmentedDynamicMesh::GetInternal(Mesh, TID),
+					Mesh.Attributes()->GetMaterialID()->GetValue(TID), TargetMaterials, true, TargetFaces);
 				if (!bIsTextureTri)
 				{
 					UV->UnsetTriangle(TID);
diff --git a/Engine/Plugins/Experimental/PlanarCutPlugin/Source/PlanarCut/Private/GeometryMeshConversion.cpp b/Engine/Plugins/Experimental/PlanarCutPlugin/Source/PlanarCut/Private/GeometryMeshConversion.cpp
--- a/Engine/Plugins/Experimental/PlanarCutPlugin/Source/PlanarCut/Private/GeometryMeshConversion.cpp
+++ b/Engine/Plugins/Experimental/PlanarCutPlugin/Source/PlanarCut/Private/GeometryMeshConversion.cpp
@@ -73,6 +73,7 @@ namespace AugmentedDynamicMesh
 	FName TangentUAttribName = "TangentUAttrib";
 	FName TangentVAttribName = "TangentVAttrib";
 	FName VisibleAttribName = "VisibleAttrib";
+	FName InternalAttribName = "InternalAttrib";
 	enum
 	{
 		MAX_NUM_UV_CHANNELS = 8,
@@ -139,6 +140,10 @@ namespace AugmentedDynamicMesh
 		TDynamicMeshScalarTriangleAttribute<bool>* VisAttrib = new TDynamicMeshScalarTriangleAttribute<bool>(&Mesh);
 		VisAttrib->Initialize(true);
 		Mesh.Attributes()->AttachAttribute(VisibleAttribName, VisAttrib);
+		TDynamicMeshScalarTriangleAttribute<bool>* InternalAttrib = new TDynamicMeshScalarTriangleAttribute<bool>(&Mesh);
+		InternalAttrib->Initialize(true);
+		Mesh.Attributes()->AttachAttribute(InternalAttribName, InternalAttrib);
+
 		EnableUVChannels(Mesh, NumUVChannels);
 	}
 
@@ -194,6 +199,22 @@ namespace AugmentedDynamicMesh
 		return Visible->GetValue(TID);
 	}
 
+	void SetInternal(FDynamicMesh3& Mesh, int TID, bool bIsInternal)
+	{
+		checkSlow(IsAugmented(Mesh));
+		TDynamicMeshScalarTriangleAttribute<bool>* Internal = 
+			static_cast<TDynamicMeshScalarTriangleAttribute<bool>*>(Mesh.Attributes()->GetAttachedAttribute(InternalAttribName));
+		Internal->SetValue(TID, bIsInternal);
+	}
+
+	bool GetInternal(const FDynamicMesh3& Mesh, int TID)
+	{
+		checkSlow(IsAugmented(Mesh));
+		TDynamicMeshScalarTriangleAttribute<bool>* Internal =
+			static_cast<TDynamicMeshScalarTriangleAttribute<bool>*>(Mesh.Attributes()->GetAttachedAttribute(InternalAttribName));
+		return Internal->GetValue(TID);
+	}
+
 	void SetUV(FDynamicMesh3& Mesh, int VID, FVector2f UV, int UVLayer)
 	{
 		if (!ensure(UVLayer < MAX_NUM_UV_CHANNELS))
@@ -465,7 +486,7 @@ namespace AugmentedDynamicMesh
 		}
 	}
 
-	void ComputeTangents(FDynamicMesh3& Mesh, bool bOnlyOddMaterials, const TArrayView<const int32>& WhichMaterials,
+	void ComputeTangents(FDynamicMesh3& Mesh, bool bOnlyInternalSurfaces, 
 		bool bRecomputeNormals, bool bMakeSharpEdges, float SharpAngleDegrees)
 	{
 		bMakeSharpEdges = bMakeSharpEdges && bRecomputeNormals; // cannot make sharp edges if normals aren't supposed to change
@@ -479,12 +500,12 @@ namespace AugmentedDynamicMesh
 		int32 NeedNumUVLayers = bMakeSharpEdges ? NumEnabledUVChannels(Mesh) : 1;
 		InitializeOverlayToPerVertexUVs(Mesh, NeedNumUVLayers, 0);
 		FDynamicMeshUVOverlay* UVs = Mesh.Attributes()->PrimaryUV();
-		FDynamicMeshMaterialAttribute* MaterialIDs = Mesh.Attributes()->GetMaterialID();
+		TDynamicMeshScalarTriangleAttribute<bool>* InternalAttrib =
+			static_cast<TDynamicMeshScalarTriangleAttribute<bool>*>(Mesh.Attributes()->GetAttachedAttribute(AugmentedDynamicMesh::InternalAttribName));
 
-		auto ShouldUpdateMID = [bOnlyOddMaterials, &WhichMaterials](int MaterialID)
+		auto ShouldUpdateInternal = [&InternalAttrib, bOnlyInternalSurfaces](int TID)
 		{
-			bool bSkipEvenMaterial = bOnlyOddMaterials && ((MaterialID % 2) == 0);
-			return !bSkipEvenMaterial && (WhichMaterials.IsEmpty() || WhichMaterials.Contains(MaterialID));
+			return !bOnlyInternalSurfaces || InternalAttrib->GetValue(TID);
 		};
 
 		// To update the normals topology, we need to weld and re-split the whole mesh
@@ -504,12 +525,12 @@ namespace AugmentedDynamicMesh
 			FMeshNormals FaceNormals(&Mesh);
 			FaceNormals.ComputeTriangleNormals();
 			Normals->CreateFromPredicate([&](int VID, int TA, int TB) {
-				int MA = MaterialIDs->GetValue(TA), MB = MaterialIDs->GetValue(TB);
-				if ((MA % 2) != (MB % 2))
+				bool IA = InternalAttrib->GetValue(TA), IB = InternalAttrib->GetValue(TB);
+				if (IA != IB)
 				{
 					return false; // always split at an internal/external face boundary
 				}
-				bool bShouldUpdateTopo = ShouldUpdateMID(MA) && ShouldUpdateMID(MB);
+				bool bShouldUpdateTopo = !bOnlyInternalSurfaces || (IA && IB);
 				if (bShouldUpdateTopo) // in the region we're updating, don't split above dot threshold
 				{
 					return FaceNormals[TA].Dot(FaceNormals[TB]) > NormalDotProdThreshold;
@@ -526,7 +547,7 @@ namespace AugmentedDynamicMesh
 			for (int TID : Mesh.TriangleIndicesItr())
 			{
 				FIndex3i OverlayTri = Normals->GetTriangle(TID);
-				if (ShouldUpdateMID(MaterialIDs->GetValue(TID)))
+				if (ShouldUpdateInternal(TID))
 				{
 					for (int SubIdx = 0; SubIdx < 3; SubIdx++)
 					{
@@ -562,7 +583,7 @@ namespace AugmentedDynamicMesh
 		const TArray<FVector3f>& TanV = Tangents.GetBitangents();
 		for (int TID : Mesh.TriangleIndicesItr())
 		{
-			if (!ShouldUpdateMID(MaterialIDs->GetValue(TID)))
+			if (!ShouldUpdateInternal(TID))
 			{
 				continue;
 			}
@@ -1158,6 +1179,7 @@ namespace AugmentedDynamicMesh
 				{
 					MaterialIDs->SetValue(TID, HoleMaterial);
 					SetVisibility(Mesh, TID, true);
+					SetInternal(Mesh, TID, true);
 				}
 			}
 		}
@@ -2558,6 +2580,7 @@ void FDynamicMeshCollection::Init(const FGeometryCollection* Collection, const T
 			}
 			Mesh.Attributes()->GetMaterialID()->SetValue(TID, Collection->MaterialID[Idx]);
 			AugmentedDynamicMesh::SetVisibility(Mesh, TID, Collection->Visible[Idx]);
+			AugmentedDynamicMesh::SetInternal(Mesh, TID, Collection->Internal[Idx]);
 			// note: material index doesn't need to be passed through; will be rebuilt by a call to reindex materials once the cut mesh is returned back to geometry collection format
 		}
 
@@ -3264,6 +3287,8 @@ bool FDynamicMeshCollection::UpdateCollection(const FTransform& FromCollection,
 		int32 CopyToIdx = FacesStart + TID;
 		Output.Visible[CopyToIdx] = AugmentedDynamicMesh::GetVisibility(Mesh, TID);
 		int MaterialID = Mesh.Attributes()->GetMaterialID()->GetValue(TID);
+		// negative material IDs are, by convention, indications of new (internal) geometry, positive are copied through
+		Output.Internal[CopyToIdx] = MaterialID < 0 ? true : AugmentedDynamicMesh::GetInternal(Mesh, TID);
 		Output.MaterialID[CopyToIdx] = MaterialID < 0 ? InternalMaterialID : MaterialID;
 		Output.Indices[CopyToIdx] = FIntVector(Mesh.GetTriangle(TID)) + VertexStartOffset;
 	}
@@ -3368,6 +3393,8 @@ int32 FDynamicMeshCollection::AppendToCollection(const FTransform& FromCollectio
 		int32 CopyToIdx = FacesStart + TID;
 		Output.Visible[CopyToIdx] = AugmentedDynamicMesh::GetVisibility(Mesh, TID);
 		int MaterialID = Mesh.Attributes()->GetMaterialID()->GetValue(TID);
+		// negative material IDs are, by convention, indications of new (internal) geometry, positive are copied through
+		Output.Internal[CopyToIdx] = MaterialID < 0 ? true : AugmentedDynamicMesh::GetInternal(Mesh, TID);
 		Output.MaterialID[CopyToIdx] = MaterialID < 0 ? InternalMaterialID : MaterialID;
 		Output.Indices[CopyToIdx] = FIntVector(Mesh.GetTriangle(TID)) + VertexStartOffset;
 	}
diff --git a/Engine/Plugins/Experimental/PlanarCutPlugin/Source/PlanarCut/Private/GeometryMeshConversion.h b/Engine/Plugins/Experimental/PlanarCutPlugin/Source/PlanarCut/Private/GeometryMeshConversion.h
--- a/Engine/Plugins/Experimental/PlanarCutPlugin/Source/PlanarCut/Private/GeometryMeshConversion.h
+++ b/Engine/Plugins/Experimental/PlanarCutPlugin/Source/PlanarCut/Private/GeometryMeshConversion.h
@@ -28,6 +28,8 @@ namespace AugmentedDynamicMesh
 {
 	void PLANARCUT_API SetVisibility(UE::Geometry::FDynamicMesh3& Mesh, int TID, bool bIsVisible);
 	bool PLANARCUT_API GetVisibility(const UE::Geometry::FDynamicMesh3& Mesh, int TID);
+	void PLANARCUT_API SetInternal(UE::Geometry::FDynamicMesh3& Mesh, int TID, bool bIsInternal);
+	bool PLANARCUT_API GetInternal(const UE::Geometry::FDynamicMesh3& Mesh, int TID);
 	void PLANARCUT_API SetUV(UE::Geometry::FDynamicMesh3& Mesh, int VID, FVector2f UV, int UVLayer);
 	void PLANARCUT_API GetUV(const UE::Geometry::FDynamicMesh3& Mesh, int VID, FVector2f& UV, int UVLayer);
 	void PLANARCUT_API SetTangent(UE::Geometry::FDynamicMesh3& Mesh, int VID, FVector3f Normal, FVector3f TangentU, FVector3f TangentV);
@@ -35,7 +37,7 @@ namespace AugmentedDynamicMesh
 	/// Initialize UV overlays based on the custom AugmentedDynamicMesh per-vertex UV attributes.  Optionally use FirstUVLayer to skip layers
 	void PLANARCUT_API InitializeOverlayToPerVertexUVs(UE::Geometry::FDynamicMesh3& Mesh, int32 NumUVLayers, int32 FirstUVLayer = 0);
 	void PLANARCUT_API InitializeOverlayToPerVertexTangents(UE::Geometry::FDynamicMesh3& Mesh);
-	void PLANARCUT_API ComputeTangents(UE::Geometry::FDynamicMesh3& Mesh, bool bOnlyOddMaterials, const TArrayView<const int32>& WhichMaterials,
+	void PLANARCUT_API ComputeTangents(UE::Geometry::FDynamicMesh3& Mesh, bool bOnlyInternalSurfaces,
 		bool bRecomputeNormals = true, bool bMakeSharpEdges = false, float SharpAngleDegrees = 60);
 	void PLANARCUT_API AddCollisionSamplesPerComponent(UE::Geometry::FDynamicMesh3& Mesh, double Spacing);
 	void PLANARCUT_API SplitOverlayAttributesToPerVertex(UE::Geometry::FDynamicMesh3& Mesh, bool bSplitUVs = true, bool bSplitTangents = true);
diff --git a/Engine/Plugins/Experimental/PlanarCutPlugin/Source/PlanarCut/Private/PlanarCut.cpp b/Engine/Plugins/Experimental/PlanarCutPlugin/Source/PlanarCut/Private/PlanarCut.cpp
--- a/Engine/Plugins/Experimental/PlanarCutPlugin/Source/PlanarCut/Private/PlanarCut.cpp
+++ b/Engine/Plugins/Experimental/PlanarCutPlugin/Source/PlanarCut/Private/PlanarCut.cpp
@@ -46,44 +46,54 @@ namespace PlanarCut_Locals
 // logic from FMeshUtility::GenerateGeometryCollectionFromBlastChunk, sets material IDs based on construction pattern that external materials have even IDs and are matched to internal materials at InternalID = ExternalID+1
 int32 FInternalSurfaceMaterials::GetDefaultMaterialIDForGeometry(const FGeometryCollection& Collection, int32 GeometryIdx) const
 {
-	int32 FaceStart = 0;
-	int32 FaceEnd = Collection.Indices.Num();
+	auto FindMostCommonInternalMaterial = [](const FGeometryCollection& InCollection, int32 InFaceStart, int32 InFaceEnd) -> int32
+	{
+		// find most common internal material
+		TMap<int32, int32> MaterialIDCount;
+		int32 MaxCount = 0;
+		int32 MostCommonMaterialID = INDEX_NONE;
+		const TManagedArray<int32>& MaterialID = InCollection.MaterialID;
+		const TManagedArray<bool>& Internal = InCollection.Internal;
+		for (int i = InFaceStart; i < InFaceEnd; ++i)
+		{
+			if (!Internal[i])
+			{
+				continue;
+			}
+			int32 CurrID = MaterialID[i];
+			int32& CurrCount = MaterialIDCount.FindOrAdd(CurrID);
+			CurrCount++;
+
+			if (CurrCount > MaxCount)
+			{
+				MaxCount = CurrCount;
+				MostCommonMaterialID = CurrID;
+			}
+		}
+		return MostCommonMaterialID;
+	};
+
+	int32 MostCommonMaterialID = INDEX_NONE;
+	int32 AllFaceStart = 0;
+	int32 AllFaceEnd = Collection.Indices.Num();
 	if (GeometryIdx > -1)
 	{
-		FaceStart = Collection.FaceStart[GeometryIdx];
-		FaceEnd = Collection.FaceCount[GeometryIdx] + Collection.FaceStart[GeometryIdx];
+		int32 GeoFaceStart = Collection.FaceStart[GeometryIdx];
+		int32 GeoFaceEnd = Collection.FaceCount[GeometryIdx] + Collection.FaceStart[GeometryIdx];
+		MostCommonMaterialID = FindMostCommonInternalMaterial(Collection, GeoFaceStart, GeoFaceEnd);
 	}
-
-	// find most common non interior material
-	TMap<int32, int32> MaterialIDCount;
-	int32 MaxCount = 0;
-	int32 MostCommonMaterialID = -1;
-	const TManagedArray<int32>&  MaterialID = Collection.MaterialID;
-	for (int i = FaceStart; i < FaceEnd; ++i)
+	if (MostCommonMaterialID == INDEX_NONE)
 	{
-		int32 CurrID = MaterialID[i];
-		int32 &CurrCount = MaterialIDCount.FindOrAdd(CurrID);
-		CurrCount++;
+		MostCommonMaterialID = FindMostCommonInternalMaterial(Collection, AllFaceStart, AllFaceEnd);
 
-		if (CurrCount > MaxCount)
+		// Failed to find any internal material IDs; default to the first material
+		if (MostCommonMaterialID == INDEX_NONE)
 		{
-			MaxCount = CurrCount;
-			MostCommonMaterialID = CurrID;
+			MostCommonMaterialID = 0;
 		}
 	}
 
-	// no face case?
-	if (MostCommonMaterialID == -1)
-	{
-		MostCommonMaterialID = 0;
-	}
-
-	// We know that the internal materials are the ones that come right after the surface materials
-	// #todo(dmp): formalize the mapping between material and internal material, perhaps on the GC
-	// if the most common material is an internal material, then just use this
-	int32 InternalMaterialID = MostCommonMaterialID % 2 == 0 ? MostCommonMaterialID + 1 : MostCommonMaterialID;
-
-	return InternalMaterialID;
+	return MostCommonMaterialID;
 }
 
 void FInternalSurfaceMaterials::SetUVScaleFromCollection(const GeometryCollection::Facades::FCollectionMeshFacade& CollectionMesh, int32 GeometryIdx)
@@ -114,14 +124,17 @@ void FInternalSurfaceMaterials::SetUVScaleFromCollection(const GeometryCollectio
 		UVDistance += FVector2D::Distance(FVector2D(UV0Array[Tri.X]), FVector2D(UV0Array[Tri.Z]));
 	}
 
+	ensure(FMath::IsFinite(GlobalUVScale));
 	if (WorldDistance > 0)
 	{
 		GlobalUVScale =  static_cast<float>(UVDistance) / WorldDistance;
 	}
+	ensure(FMath::IsFinite(GlobalUVScale));
 	if (GlobalUVScale <= 0)
 	{
 		GlobalUVScale = 1;
 	}
+	ensure(FMath::IsFinite(GlobalUVScale));
 }
 
 
@@ -1810,7 +1823,7 @@ void MergeAllSelectedBones(
 
 
 void RecomputeNormalsAndTangents(bool bOnlyTangents, bool bMakeSharpEdges, float SharpAngleDegrees, FGeometryCollection& Collection, const TArrayView<const int32>& TransformIndices,
-								 bool bOnlyOddMaterials, const TArrayView<const int32>& WhichMaterials)
+								 bool bOnlyInternalSurfaces)
 {
 	FTransform CellsToWorld = FTransform::Identity;
 
@@ -1819,7 +1832,7 @@ void RecomputeNormalsAndTangents(bool bOnlyTangents, bool bMakeSharpEdges, float
 	for (int MeshIdx = 0; MeshIdx < MeshCollection.Meshes.Num(); MeshIdx++)
 	{
 		FDynamicMesh3& Mesh = MeshCollection.Meshes[MeshIdx].AugMesh;
-		AugmentedDynamicMesh::ComputeTangents(Mesh, bOnlyOddMaterials, WhichMaterials, !bOnlyTangents, bMakeSharpEdges, SharpAngleDegrees);
+		AugmentedDynamicMesh::ComputeTangents(Mesh, bOnlyInternalSurfaces, !bOnlyTangents, bMakeSharpEdges, SharpAngleDegrees);
 	}
 
 	MeshCollection.UpdateAllCollections(Collection);
@@ -1852,7 +1865,8 @@ void ConvertToMeshDescription(
 	bool bCenterPivot,
 	FGeometryCollection& Collection,
 	const TManagedArray<FTransform>& BoneTransforms,
-	const TArrayView<const int32>& TransformIndices
+	const TArrayView<const int32>& TransformIndices,
+	TFunction<int32(int32, bool)> RemapMaterialIDs
 )
 {
 	FTransform CellsToWorld = FTransform::Identity;
@@ -1874,6 +1888,18 @@ void ConvertToMeshDescription(
 		AugmentedDynamicMesh::InitializeOverlayToPerVertexUVs(Mesh, Collection.NumUVLayers());
 		AugmentedDynamicMesh::InitializeOverlayToPerVertexTangents(Mesh);
 
+		if (RemapMaterialIDs)
+		{
+			FDynamicMeshMaterialAttribute* MaterialIDs = Mesh.Attributes()->GetMaterialID();
+			for (int32 TID : Mesh.TriangleIndicesItr())
+			{
+				int32 MatID = MaterialIDs->GetValue(TID);
+				bool bIsInternal = AugmentedDynamicMesh::GetInternal(Mesh, TID);
+				int32 NewMatID = RemapMaterialIDs(MatID, bIsInternal);
+				MaterialIDs->SetValue(TID, NewMatID);
+			}
+		}
+
 		FMergeCoincidentMeshEdges EdgeMerge(&Mesh);
 		EdgeMerge.Apply();
 		
diff --git a/Engine/Plugins/Experimental/PlanarCutPlugin/Source/PlanarCut/Public/FractureAutoUV.h b/Engine/Plugins/Experimental/PlanarCutPlugin/Source/PlanarCut/Public/FractureAutoUV.h
--- a/Engine/Plugins/Experimental/PlanarCutPlugin/Source/PlanarCut/Public/FractureAutoUV.h
+++ b/Engine/Plugins/Experimental/PlanarCutPlugin/Source/PlanarCut/Public/FractureAutoUV.h
@@ -3,6 +3,7 @@
 
 
 #include "Math/Vector4.h"
+#include "Math/MathFwd.h"
 
 class FGeometryCollection;
 namespace UE::Geometry { struct FIndex4i; }
@@ -12,6 +13,8 @@ class FProgressCancel;
 
 namespace UE { namespace PlanarCut {
 
+
+// Note: prefer ETargetFaces below
 enum class EUseMaterials
 {
 	AllMaterials,		// Include all materials
@@ -19,16 +22,48 @@ enum class EUseMaterials
 	NoDefaultMaterials	// No default materials; Only use manually selected materials
 };
 
+enum class ETargetFaces
+{
+	AllFaces,
+	InternalFaces,
+	ExternalFaces,
+	// If custom faces is chosen, then no faces are selected by default, and some other criteria (such as Material ID) must be used to select the target faces
+	CustomFaces
+};
+
+
 /**
  * Box project UVs
  *
  * @param TargetUVLayer		Which UV layer to update
  * @param Collection		The collection to be box projected
  * @param BoxDimensions		Scale of projection box
- * @param MaterialsPattern	Which pattern of material IDs to automatically consider for UV island layout
- * @param WhichMaterials	If non-empty, consider listed material IDs for UV island layout
+ * @param TargetFaces		Which faces to automatically consider for box projection
+ * @param TargetMaterials	If non-empty, additionally consider faces with listed material IDs for box projection
  */
 bool PLANARCUT_API BoxProjectUVs(
+	int32 TargetUVLayer,
+	FGeometryCollection& Collection,
+	const FVector3d& BoxDimensions,
+	ETargetFaces TargetFaces = ETargetFaces::InternalFaces,
+	TArrayView<int32> TargetMaterials = TArrayView<int32>(),
+	FVector2f OffsetUVs = FVector2f(.5f,.5f),
+	bool bOverrideBoxDimensionsWithBounds = false,
+	bool bCenterBoxAtPivot = false,
+	bool bUniformProjectionScale = false
+);
+
+/**
+ * Box project UVs
+ *
+ * @param TargetUVLayer		Which UV layer to update
+ * @param Collection		The collection to be box projected
+ * @param BoxDimensions		Scale of projection box
+ * @param MaterialsPattern	Which pattern of material IDs to automatically consider for box projection
+ * @param WhichMaterials	If non-empty, consider listed material IDs for UV box projection
+ */
+UE_DEPRECATED(5.3, "Use the BoxProjectUVs variant that takes ETargetFaces instead")
+inline bool BoxProjectUVs(
 	int32 TargetUVLayer,
 	FGeometryCollection& Collection,
 	const FVector3d& BoxDimensions,
@@ -38,6 +73,35 @@ bool PLANARCUT_API BoxProjectUVs(
 	bool bOverrideBoxDimensionsWithBounds = false,
 	bool bCenterBoxAtPivot = false,
 	bool bUniformProjectionScale = false
+)
+{
+	ETargetFaces TargetFaces =
+		MaterialsPattern == EUseMaterials::AllMaterials ? ETargetFaces::AllFaces :
+		MaterialsPattern == EUseMaterials::OddMaterials ? ETargetFaces::InternalFaces : ETargetFaces::CustomFaces;
+	return BoxProjectUVs(TargetUVLayer, Collection, BoxDimensions, TargetFaces, WhichMaterials, OffsetUVs, bOverrideBoxDimensionsWithBounds, bCenterBoxAtPivot, bUniformProjectionScale);
+}
+
+/**
+ * Make a UV atlas of non-overlapping UV charts for a geometry collection
+ *
+ * @param TargetUVLayer		Which UV layer to update with new UV coordinates
+ * @param Collection		The collection to be atlas'd
+ * @param UVRes				Target resolution for the atlas
+ * @param GutterSize		Space to leave between UV islands, in pixels at the target resolution
+ * @param TargetFaces		Which faces to automatically consider for UV island layout
+ * @param TargetMaterials	If non-empty, additionally consider faces with listed material IDs for UV island layout
+ * @param bRecreateUVsForDegenerateIslands If true, detect and fix islands that don't have proper UVs (i.e. UVs all zero or otherwise collapsed to a point)
+ */
+bool PLANARCUT_API UVLayout(
+	int32 TargetUVLayer,
+	FGeometryCollection& Collection,
+	int32 UVRes = 1024,
+	float GutterSize = 1,
+	ETargetFaces MaterialsPattern = ETargetFaces::InternalFaces,
+	TArrayView<int32> WhichMaterials = TArrayView<int32>(),
+	bool bRecreateUVsForDegenerateIslands = true,
+	FProgressCancel* Progress = nullptr
+	TArrayView<int32> WhichMaterials = TArrayView<int32>()
 );
 
 /**
@@ -51,7 +115,8 @@ bool PLANARCUT_API BoxProjectUVs(
  * @param WhichMaterials	If non-empty, consider listed material IDs for UV island layout
  * @param bRecreateUVsForDegenerateIslands If true, detect and fix islands that don't have proper UVs (i.e. UVs all zero or otherwise collapsed to a point)
  */
-bool PLANARCUT_API UVLayout(
+UE_DEPRECATED(5.3, "Use the UVLayout variant that takes ETargetFaces instead")
+inline bool UVLayout(
 	int32 TargetUVLayer,
 	FGeometryCollection& Collection,
 	int32 UVRes = 1024,
@@ -60,7 +125,13 @@ bool PLANARCUT_API UVLayout(
 	TArrayView<int32> WhichMaterials = TArrayView<int32>(),
 	bool bRecreateUVsForDegenerateIslands = true,
 	FProgressCancel* Progress = nullptr
-);
+)
+{
+	ETargetFaces TargetFaces =
+		MaterialsPattern == EUseMaterials::AllMaterials ? ETargetFaces::AllFaces :
+		MaterialsPattern == EUseMaterials::OddMaterials ? ETargetFaces::InternalFaces : ETargetFaces::CustomFaces;
+	return UVLayout(TargetUVLayer, Collection, UVRes, GutterSize, TargetFaces, WhichMaterials, bRecreateUVsForDegenerateIslands, Progress);
+}
 
 
 // Different attributes we can bake
@@ -99,8 +170,7 @@ struct FTextureAttributeSettings
 
 
 /**
- * Generate a texture for internal faces based on depth inside surface
- * TODO: add options to texture based on other quantities
+ * Generate a texture for specified groups of faces based on chosen BakeAttributes and AttributeSettings
  *
  * @param TargetUVLayer		Which UV layer to take UV coordinates from when creating the new texture
  * @param Collection		The collection to be create a new texture for
@@ -111,6 +181,20 @@ struct FTextureAttributeSettings
  * @param MaterialsPattern	Which pattern of material IDs to apply texture to
  * @param WhichMaterials	If non-empty, apply texture to the listed material IDs
  */
+bool PLANARCUT_API TextureSpecifiedFaces(
+	int32 TargetUVLayer,
+	FGeometryCollection& Collection,
+	int32 GutterSize,
+	UE::Geometry::FIndex4i BakeAttributes,
+	const FTextureAttributeSettings& AttributeSettings,
+	UE::Geometry::TImageBuilder<FVector4f>& TextureOut,
+	ETargetFaces MaterialsPattern = ETargetFaces::InternalFaces,
+	TArrayView<int32> WhichMaterials = TArrayView<int32>(),
+	FProgressCancel* Progress = nullptr
+);
+
+
+UE_DEPRECATED(5.3, "Use TextureSpecifiedFaces instead")
 bool PLANARCUT_API TextureInternalSurfaces(
 	int32 TargetUVLayer,
 	FGeometryCollection& Collection,
diff --git a/Engine/Plugins/Experimental/PlanarCutPlugin/Source/PlanarCut/Public/PlanarCut.h b/Engine/Plugins/Experimental/PlanarCutPlugin/Source/PlanarCut/Public/PlanarCut.h
--- a/Engine/Plugins/Experimental/PlanarCutPlugin/Source/PlanarCut/Public/PlanarCut.h
+++ b/Engine/Plugins/Experimental/PlanarCutPlugin/Source/PlanarCut/Public/PlanarCut.h
@@ -424,7 +424,7 @@ void PLANARCUT_API MergeAllSelectedBones(
  * @param WhichMaterials	If non-empty, restrict recomputation to only the listed material IDs
  */
 void PLANARCUT_API RecomputeNormalsAndTangents(bool bOnlyTangents, bool bMakeSharpEdges, float SharpAngleDegrees, FGeometryCollection& Collection, const TArrayView<const int32>& TransformIndices = TArrayView<const int32>(),
-	bool bOnlyOddMaterials = true, const TArrayView<const int32>& WhichMaterials = TArrayView<const int32>());
+	bool bOnlyInternalSurfaces = true);
 /**
  * Scatter additional vertices (w/ no associated triangle) as needed to satisfy minimum point spacing
  * 
@@ -498,6 +498,7 @@ inline int32 CutWithMesh(
  * @param bCenterPivot				Whether to center the geometry at the origin
  * @param Collection				The collection to be converted
  * @param TransformIndices			Which transform groups inside the collection to convert
+ * @param RemapMaterialIDs			Optional function to remap (MaterialID, bIsInternal) -> NewMaterialID
  */
 void PLANARCUT_API ConvertToMeshDescription(
 	FMeshDescription& OutputMesh,
@@ -505,7 +506,8 @@ void PLANARCUT_API ConvertToMeshDescription(
 	bool bCenterPivot,
 	FGeometryCollection& Collection,
 	const TManagedArray<FTransform>& BoneTransforms,
-	const TArrayView<const int32>& TransformIndices
+	const TArrayView<const int32>& TransformIndices,
+	TFunction<int32(int32, bool)> RemapMaterialIDs = nullptr
 );
 
 
diff --git a/Engine/Source/Runtime/Core/Public/UObject/FortniteMainBranchObjectVersion.h b/Engine/Source/Runtime/Core/Public/UObject/FortniteMainBranchObjectVersion.h
--- a/Engine/Source/Runtime/Core/Public/UObject/FortniteMainBranchObjectVersion.h
+++ b/Engine/Source/Runtime/Core/Public/UObject/FortniteMainBranchObjectVersion.h
@@ -265,6 +265,9 @@ struct CORE_API FFortniteMainBranchObjectVersion
 		// Change HLODActors to RuntimeCells mapping to use a GUID instead of the cell name
 		WorldPartitionHLODActorUseSourceCellGuid,
 
+		// Add an attribute to geometry collection to track internal faces, rather than relying on material ID numbering
+		ChaosGeometryCollectionInternalFacesAttribute,
+
 		// -----<new versions can be added above this line>-------------------------------------------------
 		VersionPlusOne,
 		LatestVersion = VersionPlusOne - 1
diff --git a/Engine/Source/Runtime/Experimental/Chaos/Private/GeometryCollection/Facades/CollectionMeshFacade.cpp b/Engine/Source/Runtime/Experimental/Chaos/Private/GeometryCollection/Facades/CollectionMeshFacade.cpp
--- a/Engine/Source/Runtime/Experimental/Chaos/Private/GeometryCollection/Facades/CollectionMeshFacade.cpp
+++ b/Engine/Source/Runtime/Experimental/Chaos/Private/GeometryCollection/Facades/CollectionMeshFacade.cpp
@@ -25,6 +25,7 @@ namespace GeometryCollection::Facades
 		, VisibleAttribute(InCollection, "Visible", FGeometryCollection::FacesGroup)
 		, MaterialIndexAttribute(InCollection, "MaterialIndex", FGeometryCollection::FacesGroup)
 		, MaterialIDAttribute(InCollection, "MaterialID", FGeometryCollection::FacesGroup)
+		, InternalAttribute(InCollection, "Internal", FGeometryCollection::FacesGroup)
 		, FaceStartAttribute(InCollection, "FaceStart", FGeometryCollection::GeometryGroup)
 		, FaceCountAttribute(InCollection, "FaceCount", FGeometryCollection::GeometryGroup)
 	{
@@ -45,6 +46,7 @@ namespace GeometryCollection::Facades
 		, VisibleAttribute(InCollection, "Visible", FGeometryCollection::FacesGroup)
 		, MaterialIndexAttribute(InCollection, "MaterialIndex", FGeometryCollection::FacesGroup)
 		, MaterialIDAttribute(InCollection, "MaterialID", FGeometryCollection::FacesGroup)
+		, InternalAttribute(InCollection, "Internal", FGeometryCollection::FacesGroup)
 		, FaceStartAttribute(InCollection, "FaceStart", FGeometryCollection::GeometryGroup)
 		, FaceCountAttribute(InCollection, "FaceCount", FGeometryCollection::GeometryGroup)
 	{
@@ -66,6 +68,7 @@ namespace GeometryCollection::Facades
 			&& VisibleAttribute.IsValid()
 			&& MaterialIndexAttribute.IsValid()
 			&& MaterialIDAttribute.IsValid()
+			&& InternalAttribute.IsValid()
 			&& FaceStartAttribute.IsValid()
 			&& FaceCountAttribute.IsValid()
 			;
@@ -87,6 +90,7 @@ namespace GeometryCollection::Facades
 		VisibleAttribute.Add();
 		MaterialIndexAttribute.Add();
 		MaterialIDAttribute.Add();
+		InternalAttribute.Add();
 		FaceStartAttribute.Add();
 		FaceCountAttribute.Add();
 	}
@@ -135,34 +139,6 @@ namespace GeometryCollection::Facades
 		return FaceIndicies;
 	}
 
-	//
-	// This is a temporary solution to decide inside/outside faces based on MaterialIDs
-	// It has to be created and set in the geometry library during fracturing operations
-	//
-	void FCollectionMeshFacade::AddInternalAttribute(FGeometryCollection& InGeometryCollection, const TArray<int32>& InMaterialID)
-	{
-		if (!InGeometryCollection.HasAttribute("Internal", FGeometryCollection::FacesGroup))
-		{
-			InGeometryCollection.AddAttribute<bool>("Internal", FGeometryCollection::FacesGroup);
-
-			TManagedArray<bool>& Internals = InGeometryCollection.ModifyAttribute<bool>("Internal", FGeometryCollection::FacesGroup);
-			const TManagedArray<int32>& MaterialIDs = InGeometryCollection.GetAttribute<int32>("MaterialID", FGeometryCollection::FacesGroup);
-
-			int32 NumFaces = InGeometryCollection.NumElements(FGeometryCollection::FacesGroup);
-			for (int32 FaceIdx = 0; FaceIdx < NumFaces; ++FaceIdx)
-			{
-				if (MaterialIDs[FaceIdx] % 2 == 1)
-				{
-					Internals[FaceIdx] = true;
-				}
-				else
-				{
-					Internals[FaceIdx] = false;
-				}
-			}
-		}
-	}
-
 	void FCollectionMeshFacade::BakeTransform(int32 TransformIdx, const FTransform& InTransform)
 	{
 		const TManagedArray<int32>& TransformToGeometryIndices = TransformToGeometryIndexAttribute.Get();
diff --git a/Engine/Source/Runtime/Experimental/Chaos/Private/GeometryCollection/GeometryCollection.cpp b/Engine/Source/Runtime/Experimental/Chaos/Private/GeometryCollection/GeometryCollection.cpp
--- a/Engine/Source/Runtime/Experimental/Chaos/Private/GeometryCollection/GeometryCollection.cpp
+++ b/Engine/Source/Runtime/Experimental/Chaos/Private/GeometryCollection/GeometryCollection.cpp
@@ -14,6 +14,7 @@
 #include "GeometryCollection/Facades/CollectionHierarchyFacade.h"
 #include "GeometryCollection/Facades/CollectionRenderingFacade.h"
 #include "UObject/FortniteNCBranchObjectVersion.h"
+#include "UObject/FortniteMainBranchObjectVersion.h"
 
 #include <iostream>
 #include <fstream>
@@ -113,6 +114,7 @@ void FGeometryCollection::Construct()
 	AddExternalAttribute<bool>("Visible", FGeometryCollection::FacesGroup, Visible);
 	AddExternalAttribute<int32>("MaterialIndex", FGeometryCollection::FacesGroup, MaterialIndex);
 	AddExternalAttribute<int32>("MaterialID", FGeometryCollection::FacesGroup, MaterialID);
+	AddExternalAttribute<bool>("Internal", FGeometryCollection::FacesGroup, Internal);
 
 	// Geometry Group
 	AddExternalAttribute<int32>("TransformIndex", FGeometryCollection::GeometryGroup, TransformIndex, TransformDependency);
@@ -176,6 +178,7 @@ int32 FGeometryCollection::AppendGeometry(const FGeometryCollection & Element, i
 	const TManagedArray<bool>& ElementVisible = Element.Visible;
 	const TManagedArray<int32>& ElementMaterialIndex = Element.MaterialIndex;
 	const TManagedArray<int32>& ElementMaterialID = Element.MaterialID;
+	const TManagedArray<bool>& ElementInternal = Element.Internal;
 
 	const TManagedArray<int32>& ElementTransformIndex = Element.TransformIndex;
 	const TManagedArray<FBox>& ElementBoundingBox = Element.BoundingBox;
@@ -250,6 +253,7 @@ int32 FGeometryCollection::AppendGeometry(const FGeometryCollection & Element, i
 	{
 		Indices[IndicesIndex + tdx] = FIntVector(VerticesIndex, VerticesIndex, VerticesIndex) + ElementIndices[tdx];
 		Visible[IndicesIndex + tdx] = ElementVisible[tdx];
+		Internal[IndicesIndex + tdx] = ElementInternal[tdx];
 		MaterialIndex[IndicesIndex + tdx] = ElementMaterialIndex[tdx];
 		// MaterialIDs need to be incremented
 		MaterialID[IndicesIndex + tdx] = MaterialIDOffset + ElementMaterialID[tdx];	
@@ -996,6 +1000,7 @@ void FGeometryCollection::Serialize(Chaos::FChaosArchive& Ar)
 	}
 
 	Ar.UsingCustomVersion(FFortniteNCBranchObjectVersion::GUID);
+	Ar.UsingCustomVersion(FFortniteMainBranchObjectVersion::GUID);
 
 	Super::Serialize(Ar);
 
@@ -1312,6 +1317,19 @@ void FGeometryCollection::Serialize(Chaos::FChaosArchive& Ar)
 			Version = 10;
 		}
 
+		if (Ar.CustomVer(FFortniteMainBranchObjectVersion::GUID) < FFortniteMainBranchObjectVersion::ChaosGeometryCollectionInternalFacesAttribute)
+		{
+			if (!HasAttribute("Internal", FacesGroup))
+			{
+				AddExternalAttribute<bool>("Internal", FacesGroup, Internal);
+			}
+
+			for (int32 FaceIdx = 0; FaceIdx < MaterialID.Num(); ++FaceIdx)
+			{
+				Internal[FaceIdx] = bool(MaterialID[FaceIdx] & 1);
+			}
+		}
+
 		// Finally, make sure expected interfaces are initialized
 		InitializeInterfaces();
 	}
@@ -1514,6 +1532,7 @@ void FGeometryCollection::Init(FGeometryCollection* Collection, const TArray<flo
 		TManagedArray<bool>& Visible = Collection->Visible;
 		TManagedArray<int32>& MaterialID = Collection->MaterialID;
 		TManagedArray<int32>& MaterialIndex = Collection->MaterialIndex;
+		TManagedArray<bool>& Internal = Collection->Internal;
 		TManagedArray<FTransform>& Transform = Collection->Transform;
 		
 		Collection->SetNumUVLayers(1);
@@ -1558,6 +1577,7 @@ void FGeometryCollection::Init(FGeometryCollection* Collection, const TArray<flo
 
 			Indices[Idx] = FIntVector(VertexIdx1, VertexIdx2, VertexIdx3);
 			Visible[Idx] = true;
+			Internal[Idx] = false;
 			MaterialID[Idx] = 0;
 			MaterialIndex[Idx] = Idx;
 
@@ -1832,6 +1852,7 @@ FGeometryCollection* FGeometryCollection::NewGeometryCollection(const TArray<flo
 	TManagedArray<bool>&  Visible = RestCollection->Visible;
 	TManagedArray<int32>&  MaterialID = RestCollection->MaterialID;
 	TManagedArray<int32>&  MaterialIndex = RestCollection->MaterialIndex;
+	TManagedArray<bool>& Internal = RestCollection->Internal;
 	TManagedArray<FTransform>&  Transform = RestCollection->Transform;
 	TManagedArray<int32>& Parent = RestCollection->Parent;
 	TManagedArray<TSet<int32>>& Children = RestCollection->Children;
@@ -1891,6 +1912,7 @@ FGeometryCollection* FGeometryCollection::NewGeometryCollection(const TArray<flo
 
 		Indices[Idx] = FIntVector(VertexIdx1, VertexIdx2, VertexIdx3);
 		Visible[Idx] = true;
+		Internal[Idx] = false;
 		MaterialID[Idx] = 0;
 		MaterialIndex[Idx] = Idx;
 
diff --git a/Engine/Source/Runtime/Experimental/Chaos/Private/GeometryCollection/GeometryCollectionAlgo.cpp b/Engine/Source/Runtime/Experimental/Chaos/Private/GeometryCollection/GeometryCollectionAlgo.cpp
--- a/Engine/Source/Runtime/Experimental/Chaos/Private/GeometryCollection/GeometryCollectionAlgo.cpp
+++ b/Engine/Source/Runtime/Experimental/Chaos/Private/GeometryCollection/GeometryCollectionAlgo.cpp
@@ -802,6 +802,8 @@ namespace GeometryCollectionAlgo
 				// just copy material from one of the faces in the group
 				GeometryCollection->MaterialIndex[NewSpaceIdx] = GeometryCollection->MaterialIndex[OldStart];
 				GeometryCollection->MaterialID[NewSpaceIdx] = GeometryCollection->MaterialID[OldStart];
+				// also copy internal status from the existing faces
+				GeometryCollection->Internal[NewSpaceIdx] = GeometryCollection->Internal[OldStart];
 				ShiftedOrder[ShiftedIdx++] = NewSpaceIdx;
 				NewSpaceIdx++;
 			}
diff --git a/Engine/Source/Runtime/Experimental/Chaos/Private/GeometryCollection/GeometryCollectionProximityPropertiesInterface.cpp b/Engine/Source/Runtime/Experimental/Chaos/Private/GeometryCollection/GeometryCollectionProximityPropertiesInterface.cpp
--- a/Engine/Source/Runtime/Experimental/Chaos/Private/GeometryCollection/GeometryCollectionProximityPropertiesInterface.cpp
+++ b/Engine/Source/Runtime/Experimental/Chaos/Private/GeometryCollection/GeometryCollectionProximityPropertiesInterface.cpp
@@ -86,8 +86,6 @@ FGeometryCollectionProximityPropertiesInterface::GetProximityProperties() const
 void
 FGeometryCollectionProximityPropertiesInterface::SetProximityProperties(const FProximityProperties& InProximityAttributes)
 {
-	FProximityProperties Property;
-
 	int32 AttributeIndex = 0;
 	const bool bHasProximityProperties = ManagedCollection->NumElements(ProximityPropertiesGroup) > 0;
 	if (!bHasProximityProperties)
diff --git a/Engine/Source/Runtime/Experimental/Chaos/Private/GeometryCollection/GeometryCollectionUtility.cpp b/Engine/Source/Runtime/Experimental/Chaos/Private/GeometryCollection/GeometryCollectionUtility.cpp
--- a/Engine/Source/Runtime/Experimental/Chaos/Private/GeometryCollection/GeometryCollectionUtility.cpp
+++ b/Engine/Source/Runtime/Experimental/Chaos/Private/GeometryCollection/GeometryCollectionUtility.cpp
@@ -38,6 +38,7 @@ namespace GeometryCollection
 		TManagedArray<bool>&  Visible = RestCollection->Visible;
 		TManagedArray<int32>&  MaterialIndex = RestCollection->MaterialIndex;
 		TManagedArray<int32>&  MaterialID = RestCollection->MaterialID;
+		TManagedArray<bool>& Internal = RestCollection->Internal;
 		TManagedArray<FTransform>&  Transform = RestCollection->Transform;
 		TManagedArray<int32>& SimType = RestCollection->SimulationType;
 
@@ -110,9 +111,10 @@ namespace GeometryCollection
 		// distribute the number of materials equally between the 12 faces
 		check(NumberOfMaterials <= 12 && (12 % NumberOfMaterials)==0); // preferably divisible into 12
 		int NumberOfEachMaterial = 12 / NumberOfMaterials;
-		for (int i = 0; i < 12;i++)
+		for (int i = 0; i < 12; i++)
 		{
 			Visible[i] = true;
+			Internal[i] = false;
 
 			MaterialIndex[i] = i;
 			MaterialID[i] = i / NumberOfEachMaterial;
diff --git a/Engine/Source/Runtime/Experimental/Chaos/Public/GeometryCollection/Facades/CollectionMeshFacade.h b/Engine/Source/Runtime/Experimental/Chaos/Public/GeometryCollection/Facades/CollectionMeshFacade.h
--- a/Engine/Source/Runtime/Experimental/Chaos/Public/GeometryCollection/Facades/CollectionMeshFacade.h
+++ b/Engine/Source/Runtime/Experimental/Chaos/Public/GeometryCollection/Facades/CollectionMeshFacade.h
@@ -42,11 +42,6 @@ namespace GeometryCollection::Facades
 		 */
 		const TArray<int32> GetFaceIndices(int32 BoneIdx) const;
 
-		/**
-		 * Adds and sets "Internal" attribute in FGeometryCollection::FacesGroup to designate internal faces
-		 */
-		static void AddInternalAttribute(FGeometryCollection& InGeometryCollection, const TArray<int32>& InMaterialID);
-
 		/**
 		 * Returns the vertex indices of the face for the bone
 		 */
@@ -70,6 +65,7 @@ namespace GeometryCollection::Facades
 		TManagedArrayAccessor<bool> VisibleAttribute;
 		TManagedArrayAccessor<int32> MaterialIndexAttribute;
 		TManagedArrayAccessor<int32> MaterialIDAttribute;
+		TManagedArrayAccessor<bool> InternalAttribute;
 		TManagedArrayAccessor<int32> FaceStartAttribute;
 		TManagedArrayAccessor<int32> FaceCountAttribute;
 	};
diff --git a/Engine/Source/Runtime/Experimental/Chaos/Public/GeometryCollection/GeometryCollection.h b/Engine/Source/Runtime/Experimental/Chaos/Public/GeometryCollection/GeometryCollection.h
--- a/Engine/Source/Runtime/Experimental/Chaos/Public/GeometryCollection/GeometryCollection.h
+++ b/Engine/Source/Runtime/Experimental/Chaos/Public/GeometryCollection/GeometryCollection.h
@@ -311,6 +311,7 @@ public:
 	TManagedArray<bool>         Visible;
 	TManagedArray<int32>        MaterialIndex;
 	TManagedArray<int32>        MaterialID;
+	TManagedArray<bool>         Internal;
 
 	// Geometry Group
 	TManagedArray<int32>        TransformIndex;
diff --git a/Engine/Source/Runtime/Experimental/Chaos/Public/GeometryCollection/GeometryCollectionUtility.h b/Engine/Source/Runtime/Experimental/Chaos/Public/GeometryCollection/GeometryCollectionUtility.h
--- a/Engine/Source/Runtime/Experimental/Chaos/Public/GeometryCollection/GeometryCollectionUtility.h
+++ b/Engine/Source/Runtime/Experimental/Chaos/Public/GeometryCollection/GeometryCollectionUtility.h
@@ -45,6 +45,7 @@ namespace GeometryCollection
 		TManagedArray<bool>& Visible = RestCollection->Visible;
 		TManagedArray<int32>& MaterialIndex = RestCollection->MaterialIndex;
 		TManagedArray<int32>& MaterialID = RestCollection->MaterialID;
+		TManagedArray<bool>& Internal = RestCollection->Internal;
 		TManagedArray<FTransform>& Transform = RestCollection->Transform;
 		TManagedArray<int32>& SimulationType = RestCollection->SimulationType;
 
@@ -73,6 +74,7 @@ namespace GeometryCollection
 			Indices[Idx] = FIntVector(Tri[0], Tri[1], Tri[2]);
 
 			Visible[Idx] = true;
+			Internal[Idx] = false;
 			MaterialIndex[Idx] = Idx;
 			MaterialID[Idx] = Idx / NumberOfEachMaterial;
 
diff --git a/Engine/Source/Runtime/Experimental/GeometryCollectionEngine/Private/GeometryCollection/GeometryCollectionComponent.cpp b/Engine/Source/Runtime/Experimental/GeometryCollectionEngine/Private/GeometryCollection/GeometryCollectionComponent.cpp
--- a/Engine/Source/Runtime/Experimental/GeometryCollectionEngine/Private/GeometryCollection/GeometryCollectionComponent.cpp
+++ b/Engine/Source/Runtime/Experimental/GeometryCollectionEngine/Private/GeometryCollection/GeometryCollectionComponent.cpp
@@ -2130,6 +2130,7 @@ void UGeometryCollectionComponent::InitConstantData(TUniquePtr<FGeometryCollecti
 		int32 NumIndices = 0;
 		const TManagedArray<FIntVector>& Indices = Collection->Indices;
 		const TManagedArray<int32>& MaterialID = Collection->MaterialID;
+		const TManagedArray<bool>& Internal = Collection->Internal;
 
 		const TManagedArray<bool>& Visible = GetVisibleArray();  // Use copy on write attribute. The rest collection visible array can be overriden for the convenience of debug drawing the collision volumes
 		
@@ -2259,12 +2260,12 @@ void UGeometryCollectionComponent::InitConstantData(TUniquePtr<FGeometryCollecti
 		// #note:  This is a stopgap because the original geometry array is broken
 		for (int FaceIndex = 0; FaceIndex < NumFaces; ++FaceIndex)
 		{
-			// only add visible external faces.  MaterialID that is even is an external material
+			// only add visible external faces.
 #if WITH_EDITOR
 			const bool bUseVisible = bUsingHideArray ? VisibleOverride[FaceIndex] : Visible[FaceIndex];
-			if (bUseVisible && (MaterialID[FaceIndex] % 2 == 0 || bUsingHideArray))
+			if (bUseVisible && (Internal[FaceIndex] == false || bUsingHideArray))
 #else
-			if (Visible[FaceIndex] && MaterialID[FaceIndex] % 2 == 0)
+			if (Visible[FaceIndex] && Internal[FaceIndex] == false)
 #endif
 			{
 				BaseMeshIndices.Add(Indices[FaceIndex]);
diff --git a/Engine/Source/Runtime/Experimental/GeometryCollectionEngine/Private/GeometryCollection/GeometryCollectionEngineConversion.cpp b/Engine/Source/Runtime/Experimental/GeometryCollectionEngine/Private/GeometryCollection/GeometryCollectionEngineConversion.cpp
--- a/Engine/Source/Runtime/Experimental/GeometryCollectionEngine/Private/GeometryCollection/GeometryCollectionEngineConversion.cpp
+++ b/Engine/Source/Runtime/Experimental/GeometryCollectionEngine/Private/GeometryCollection/GeometryCollectionEngineConversion.cpp
@@ -93,7 +93,7 @@ static FVector GetMeshBuildScale3D(const UStaticMesh& StaticMesh)
 
 void FGeometryCollectionEngineConversion::AppendMeshDescription(
 	const FMeshDescription* MeshDescription, const FString& Name, int32 MaterialStartIndex, const FTransform& StaticMeshTransform, 
-	FGeometryCollection* GeometryCollection, UBodySetup* BodySetup, bool ReindexMaterials, bool bAddInternalMaterials)
+	FGeometryCollection* GeometryCollection, UBodySetup* BodySetup, bool ReindexMaterials, bool bAddInternalMaterials, bool bSetInternalFromMaterialIndex)
 {
 #if WITH_EDITORONLY_DATA
 
@@ -207,6 +207,7 @@ void FGeometryCollectionEngineConversion::AppendMeshDescription(
 	TManagedArray<bool>& TargetVisible = GeometryCollection->Visible;
 	TManagedArray<int32>& TargetMaterialID = GeometryCollection->MaterialID;
 	TManagedArray<int32>& TargetMaterialIndex = GeometryCollection->MaterialIndex;
+	TManagedArray<bool>& TargetInternal = GeometryCollection->Internal;
 
 	const int32 IndicesCount = MeshDescription->Triangles().Num();
 	const int32 InitialNumIndices = GeometryCollection->NumElements(FGeometryCollection::FacesGroup);
@@ -224,9 +225,17 @@ void FGeometryCollectionEngineConversion::AppendMeshDescription(
 
 		TargetVisible[TargetIndex] = true;
 
-		// If adding internal materials, then materials are ganged in pairs and we want the id to associate with the first of each pair.
+		// bAddInternalMaterials and bSetInternalFromMaterialIndex support the legacy system of odd-numbered materials indicating internal surfaces
+		// bSetInternalFromMaterialIndex can be used to 'round trip' which faces are internal when using the 'ToMesh' tool to temporarily convert to/from static mesh.
 		int32 MaterialIndexScale = 1 + int32(bAddInternalMaterials);
-		TargetMaterialID[TargetIndex] = MaterialStartIndex + (MeshDescription->GetTrianglePolygonGroup(TriangleIndex) * MaterialIndexScale);
+		int32 MaterialSourceID = MeshDescription->GetTrianglePolygonGroup(TriangleIndex);
+		TargetMaterialID[TargetIndex] = MaterialStartIndex + (MaterialSourceID * MaterialIndexScale);
+		bool bIsInternal = false;
+		if (bSetInternalFromMaterialIndex && !bAddInternalMaterials)
+		{
+			bIsInternal = (MaterialSourceID % 2) == 1;
+		}
+		TargetInternal[TargetIndex] = bIsInternal;
 
 		// Is this right?
 		TargetMaterialIndex[TargetIndex] = TargetIndex;
@@ -488,7 +497,7 @@ void FGeometryCollectionEngineConversion::AppendAutoInstanceMeshIndices(UGeometr
 
 bool FGeometryCollectionEngineConversion::AppendStaticMesh(const UStaticMesh* StaticMesh, const TArray<UMaterialInterface*>& Materials,
 	const FTransform& StaticMeshTransform, UGeometryCollection* GeometryCollectionObject, bool bReindexMaterials,
-	bool bAddInternalMaterials, bool bSplitComponents)
+	bool bAddInternalMaterials, bool bSplitComponents, bool bSetInternalFromMaterialIndex)
 {
 #if WITH_EDITORONLY_DATA
 
@@ -501,7 +510,7 @@ bool FGeometryCollectionEngineConversion::AppendStaticMesh(const UStaticMesh* St
 
 	const int32 OriginalNumOfTransforms = GeometryCollection->NumElements(FGeometryCollection::TransformGroup);
 
-	if (AppendStaticMesh(StaticMesh, StartMaterialIndex, StaticMeshTransform, GeometryCollection, bReindexMaterials, bAddInternalMaterials, bSplitComponents))
+	if (AppendStaticMesh(StaticMesh, StartMaterialIndex, StaticMeshTransform, GeometryCollection, bReindexMaterials, bAddInternalMaterials, bSplitComponents, bSetInternalFromMaterialIndex))
 	{
 		AppendMaterials(Materials, GeometryCollectionObject, bAddInternalMaterials);
 
@@ -515,7 +524,7 @@ bool FGeometryCollectionEngineConversion::AppendStaticMesh(const UStaticMesh* St
 }
 
 bool FGeometryCollectionEngineConversion::AppendStaticMesh(const UStaticMesh* StaticMesh, int32 StartMaterialIndex,  const FTransform& StaticMeshTransform,
-	FGeometryCollection* GeometryCollection, bool bReindexMaterials, bool bAddInternalMaterials, bool bSplitComponents)
+	FGeometryCollection* GeometryCollection, bool bReindexMaterials, bool bAddInternalMaterials, bool bSplitComponents, bool bSetInternalFromMaterialIndex)
 {
 #if WITH_EDITORONLY_DATA
 	if (StaticMesh)
@@ -646,7 +655,7 @@ bool FGeometryCollectionEngineConversion::AppendStaticMesh(const UStaticMesh* St
 
 					for (FMeshDescription& MD : Descriptions)
 					{
-						AppendMeshDescription(&MD, StaticMesh->GetName(), StartMaterialIndex, MeshTransform, GeometryCollection, StaticMesh->GetBodySetup(), false, bAddInternalMaterials);
+						AppendMeshDescription(&MD, StaticMesh->GetName(), StartMaterialIndex, MeshTransform, GeometryCollection, StaticMesh->GetBodySetup(), false, bAddInternalMaterials, bSetInternalFromMaterialIndex);
 					}
 
 					if (bReindexMaterials)
@@ -659,7 +668,7 @@ bool FGeometryCollectionEngineConversion::AppendStaticMesh(const UStaticMesh* St
 				// else only one component -- fall back to just using the original mesh description
 			}
 
-			AppendMeshDescription(MeshDescription, StaticMesh->GetName(), StartMaterialIndex, MeshTransform, GeometryCollection, StaticMesh->GetBodySetup(), bReindexMaterials, bAddInternalMaterials);
+			AppendMeshDescription(MeshDescription, StaticMesh->GetName(), StartMaterialIndex, MeshTransform, GeometryCollection, StaticMesh->GetBodySetup(), bReindexMaterials, bAddInternalMaterials, bSetInternalFromMaterialIndex);
 			return true;
 		}
 	}
@@ -736,12 +745,14 @@ bool FGeometryCollectionEngineConversion::AppendGeometryCollection(const FGeomet
 	const TManagedArray<bool>& SourceVisible = SourceGeometryCollectionPtr->Visible;
 	const TManagedArray<int32>& SourceMaterialID = SourceGeometryCollectionPtr->MaterialID;
 	const TManagedArray<int32>& SourceMaterialIndex = SourceGeometryCollectionPtr->MaterialIndex;
+	const TManagedArray<bool>& SourceInternal = SourceGeometryCollectionPtr->Internal;
 
 	// target face information
 	TManagedArray<FIntVector>& TargetIndices = TargetGeometryCollection->Indices;
 	TManagedArray<bool>& TargetVisible = TargetGeometryCollection->Visible;
 	TManagedArray<int32>& TargetMaterialID = TargetGeometryCollection->MaterialID;
 	TManagedArray<int32>& TargetMaterialIndex = TargetGeometryCollection->MaterialIndex;
+	TManagedArray<bool>& TargetInternal = TargetGeometryCollection->Internal;
 
 	// append faces
 	for (int32 FaceIndex = 0; FaceIndex < FaceCount; ++FaceIndex)
@@ -756,6 +767,7 @@ bool FGeometryCollectionEngineConversion::AppendGeometryCollection(const FGeomet
 
 		TargetMaterialID[FaceOffset] = AssetMaterialStart + SourceMaterialID[FaceIndex];
 		TargetMaterialIndex[FaceOffset] = FaceOffset;
+		TargetInternal[FaceOffset] = SourceInternal[FaceIndex];
 	}
 
 	using FCollisionType = FGeometryDynamicCollection::FSharedImplicit;
@@ -931,7 +943,7 @@ void FGeometryCollectionEngineConversion::AppendGeometryCollection(const UGeomet
 
 
 void FGeometryCollectionEngineConversion::AppendStaticMesh(const UStaticMesh* StaticMesh, const UStaticMeshComponent* StaticMeshComponent, const FTransform& StaticMeshTransform, UGeometryCollection* GeometryCollectionObject, 
-	bool ReindexMaterials, bool bAddInternalMaterials, bool bSplitComponents)
+	bool ReindexMaterials, bool bAddInternalMaterials, bool bSplitComponents, bool bSetInternalFromMaterialIndex)
 {
 	if (StaticMesh == nullptr)
 	{
@@ -952,7 +964,7 @@ void FGeometryCollectionEngineConversion::AppendStaticMesh(const UStaticMesh* St
 	GeometryCollectionObject->Materials.Remove(BoneSelectedMaterial);
 	Materials.Remove(BoneSelectedMaterial);
 	
-	AppendStaticMesh(StaticMesh, Materials, StaticMeshTransform, GeometryCollectionObject, ReindexMaterials, bAddInternalMaterials, bSplitComponents);
+	AppendStaticMesh(StaticMesh, Materials, StaticMeshTransform, GeometryCollectionObject, ReindexMaterials, bAddInternalMaterials, bSplitComponents, bSetInternalFromMaterialIndex);
 }
 
 
@@ -1345,14 +1357,16 @@ void FGeometryCollectionEngineConversion::AppendGeometryCollectionSource(const F
 
 		if (const UStaticMesh* SourceStaticMesh = Cast<UStaticMesh>(SourceObject))
 		{
+			bool bLegacyAddInternal = GeometryCollectionSource.bAddInternalMaterials;
 			AppendStaticMesh(
 				SourceStaticMesh,
 				StartMaterialIndex,
 				GeometryCollectionSource.LocalTransform,
 				&GeometryCollectionInOut,
 				ReindexMaterials,
-				GeometryCollectionSource.bAddInternalMaterials,
-				GeometryCollectionSource.bSplitComponents
+				bLegacyAddInternal,
+				GeometryCollectionSource.bSplitComponents,
+				GeometryCollectionSource.bSetInternalFromMaterialIndex
 				);
 		}
 		else if (const USkeletalMesh* SourceSkeletalMesh = Cast<USkeletalMesh>(SourceObject))
diff --git a/Engine/Source/Runtime/Experimental/GeometryCollectionEngine/Private/GeometryCollection/GeometryCollectionObject.cpp b/Engine/Source/Runtime/Experimental/GeometryCollectionEngine/Private/GeometryCollection/GeometryCollectionObject.cpp
--- a/Engine/Source/Runtime/Experimental/GeometryCollectionEngine/Private/GeometryCollection/GeometryCollectionObject.cpp
+++ b/Engine/Source/Runtime/Experimental/GeometryCollectionEngine/Private/GeometryCollection/GeometryCollectionObject.cpp
@@ -494,7 +494,7 @@ void UGeometryCollection::ReindexMaterialSections()
 	InvalidateCollection();
 }
 
-void UGeometryCollection::InitializeMaterials(bool bHasInternalMaterials)
+void UGeometryCollection::InitializeMaterials(bool bHasLegacyInternalMaterialsPairs)
 {
 	Modify();
 
@@ -516,7 +516,7 @@ void UGeometryCollection::InitializeMaterials(bool bHasInternalMaterials)
 	}
 
 	TArray<UMaterialInterface*> FinalMaterials;
-	if (bHasInternalMaterials)
+	if (bHasLegacyInternalMaterialsPairs)
 	{
 		// We're assuming that all materials are arranged in pairs, so first we collect these.
 		using FMaterialPair = TPair<UMaterialInterface*, UMaterialInterface*>;
@@ -611,7 +611,32 @@ void UGeometryCollection::InitializeMaterials(bool bHasInternalMaterials)
 	InvalidateCollection();
 }
 
+int32 UGeometryCollection::AddNewMaterialSlot(bool bCopyLastMaterial)
+{
+	Modify();
+	int32 NewIdx = Materials.Emplace();
+	if (NewIdx > 0 && bCopyLastMaterial)
+	{
+		Materials[NewIdx] = Materials[NewIdx - 1];
+	}
+
+	InvalidateCollection();
+
+	return NewIdx;
+}
 
+bool UGeometryCollection::RemoveLastMaterialSlot()
+{
+	if (Materials.Num() > 1)
+	{
+		Modify();
+		Materials.Pop();
+		InvalidateCollection();
+		return true;
+	}
+
+	return false;
+}
 
 /** Returns true if there is anything to render */
 bool UGeometryCollection::HasVisibleGeometry() const
diff --git a/Engine/Source/Runtime/Experimental/GeometryCollectionEngine/Public/GeometryCollection/GeometryCollectionComponent.h b/Engine/Source/Runtime/Experimental/GeometryCollectionEngine/Public/GeometryCollection/GeometryCollectionComponent.h
--- a/Engine/Source/Runtime/Experimental/GeometryCollectionEngine/Public/GeometryCollection/GeometryCollectionComponent.h
+++ b/Engine/Source/Runtime/Experimental/GeometryCollectionEngine/Public/GeometryCollection/GeometryCollectionComponent.h
@@ -684,6 +684,7 @@ public:
 	COPY_ON_WRITE_ATTRIBUTE(bool, Visible, FGeometryCollection::FacesGroup)								\
 	COPY_ON_WRITE_ATTRIBUTE(int32, MaterialIndex, FGeometryCollection::FacesGroup)						\
 	COPY_ON_WRITE_ATTRIBUTE(int32, MaterialID, FGeometryCollection::FacesGroup)							\
+	COPY_ON_WRITE_ATTRIBUTE(bool, Internal, FGeometryCollection::FacesGroup)							\
 																										\
 	/* Geometry Group */																				\
 	COPY_ON_WRITE_ATTRIBUTE(int32, TransformIndex, FGeometryCollection::GeometryGroup)					\
diff --git a/Engine/Source/Runtime/Experimental/GeometryCollectionEngine/Public/GeometryCollection/GeometryCollectionEngineConversion.h b/Engine/Source/Runtime/Experimental/GeometryCollectionEngine/Public/GeometryCollection/GeometryCollectionEngineConversion.h
--- a/Engine/Source/Runtime/Experimental/GeometryCollectionEngine/Public/GeometryCollection/GeometryCollectionEngineConversion.h
+++ b/Engine/Source/Runtime/Experimental/GeometryCollectionEngine/Public/GeometryCollection/GeometryCollectionEngineConversion.h
@@ -54,7 +54,7 @@ public:
 	 * @param bReindexMaterials	: Whether to reindex materials -- if appending multiple meshes, pass false and call ReindexMaterials afterwards
 	 */
 	static void AppendMeshDescription(const FMeshDescription* MeshDescription, const FString& Name, int32 StartMaterialIndex, const FTransform& StaticMeshTransform,
-		FGeometryCollection* GeometryCollection, UBodySetup* BodySetup = nullptr, bool bReindexMaterials = true, bool bAddInternalMaterials = true);
+		FGeometryCollection* GeometryCollection, UBodySetup* BodySetup = nullptr, bool bReindexMaterials = true, bool bAddInternalMaterials = true, bool bSetInternalFromMaterialIndex = false);
 
 	/**
 	 * Get a HiRes (or LOD 0 if no HiRes available) MeshDescription for the given static mesh, and make sure it includes normals and tangents.
@@ -72,7 +72,7 @@ public:
 	 * @return true if succeeded in appending a mesh
 	 */
 	static bool AppendStaticMesh(const UStaticMesh* StaticMesh, const TArray<UMaterialInterface*>& Materials, const FTransform& StaticMeshTransform, 
-		UGeometryCollection* GeometryCollectionObject, bool bReindexMaterials = true, bool bAddInternalMaterials = true, bool bSplitComponents = false);
+		UGeometryCollection* GeometryCollectionObject, bool bReindexMaterials = true, bool bAddInternalMaterials = true, bool bSplitComponents = false, bool bSetInternalFromMaterialIndex = false);
 
 	// TODO: consider having the below fn return bool, and have a helper that returns the StartMaterialIndex from UGeometryCollection
 	//  before any materials are actually added, so we can call the AppendMaterial function after, and not need to use its return value ...
@@ -89,7 +89,7 @@ public:
 	 * @return true if succeeded in appending a mesh
 	 */
 	static bool AppendStaticMesh(const UStaticMesh* StaticMesh, int32 StartMaterialIndex, const FTransform& StaticMeshTransform, 
-		FGeometryCollection* GeometryCollection, bool bReindexMaterials = true, bool bAddInternalMaterials = true, bool bSplitComponents = false);
+		FGeometryCollection* GeometryCollection, bool bReindexMaterials = true, bool bAddInternalMaterials = true, bool bSplitComponents = false, bool bSetInternalFromMaterialIndex = false);
 
 	/**
 	*  Appends a static mesh to a GeometryCollectionComponent.
@@ -98,7 +98,7 @@ public:
 	*  @param GeometryCollection  : Collection to append the mesh into.
 	*/
 	static void AppendStaticMesh(const UStaticMesh* StaticMesh, const UStaticMeshComponent *StaticMeshComponent, const FTransform& StaticMeshTransform, UGeometryCollection* GeometryCollection,
-		bool bReindexMaterials = true, bool bAddInternalMaterials = true, bool bSplitComponents = false);
+		bool bReindexMaterials = true, bool bAddInternalMaterials = true, bool bSplitComponents = false, bool bSetInternalFromMaterialIndex = false);
 
 	/**
 	*  Appends a skeletal mesh to a GeometryCollectionComponent.
diff --git a/Engine/Source/Runtime/Experimental/GeometryCollectionEngine/Public/GeometryCollection/GeometryCollectionObject.h b/Engine/Source/Runtime/Experimental/GeometryCollectionEngine/Public/GeometryCollection/GeometryCollectionObject.h
--- a/Engine/Source/Runtime/Experimental/GeometryCollectionEngine/Public/GeometryCollection/GeometryCollectionObject.h
+++ b/Engine/Source/Runtime/Experimental/GeometryCollectionEngine/Public/GeometryCollection/GeometryCollectionObject.h
@@ -26,6 +26,13 @@ struct GEOMETRYCOLLECTIONENGINE_API FGeometryCollectionSource
 {
 	GENERATED_BODY()
 
+	FGeometryCollectionSource() {}
+	FGeometryCollectionSource(const FSoftObjectPath& SourceSoftObjectPath, const FTransform& ComponentTransform, const TArray<TObjectPtr<UMaterialInterface>>& SourceMaterials, bool bSplitComponents = false, bool bSetInternalFromMaterialIndex = false)
+		: SourceGeometryObject(SourceSoftObjectPath), LocalTransform(ComponentTransform), SourceMaterial(SourceMaterials), bAddInternalMaterials(false), bSplitComponents(bSplitComponents), bSetInternalFromMaterialIndex(bSetInternalFromMaterialIndex)
+	{
+
+	}
+
 	UPROPERTY(BlueprintReadWrite, EditAnywhere, Category = "GeometrySource", meta=(AllowedClasses="/Script/Engine.StaticMesh, /Script/Engine.SkeletalMesh, /Script/GeometryCollectionEngine.GeometryCollection"))
 	FSoftObjectPath SourceGeometryObject;
 
@@ -35,14 +42,19 @@ struct GEOMETRYCOLLECTIONENGINE_API FGeometryCollectionSource
 	UPROPERTY(BlueprintReadWrite, EditAnywhere, Category = "GeometrySource")
 	TArray<TObjectPtr<UMaterialInterface>> SourceMaterial;
 
-	/** Whether source materials should be duplicated to create slots for internal materials. Does not apply if the source is a GeometryCollection. */
-	UPROPERTY(BlueprintReadWrite, EditAnywhere, Category = "GeometrySource")
+	//~ Note: bAddInternalMaterials defaults to true so a 'Reset' of a geometry collection that was created before this member was added will have consistent behavior. New geometry collections should always set bAddInternalMaterials to false.
+	/** (Legacy) Whether source materials will be duplicated to create new slots for internal materials, or existing odd materials will be considered internal. (For non-Geometry Collection inputs only.) */
+	UPROPERTY(BlueprintReadWrite, EditAnywhere, Category = "GeometrySource", DisplayName = "(Legacy) Add Internal Materials")
 	bool bAddInternalMaterials = true;
 
-	/** Whether individual source mesh components should be split into separate pieces of geometry based on mesh connectivity. If checked, triangles that are not topologically connected will be assigned separate bones. */
+	/** Whether individual source mesh components should be split into separate pieces of geometry based on mesh connectivity. If checked, triangles that are not topologically connected will be assigned separate bones. (For non-Geometry Collection inputs only.) */
 	UPROPERTY(BlueprintReadWrite, EditAnywhere, Category = "GeometrySource")
 	bool bSplitComponents = false;
 
+	/** Whether to set the 'internal' flag for faces with odd-numbered materials slots. */
+	UPROPERTY(BlueprintReadWrite, EditAnywhere, Category = "GeometrySource")
+	bool bSetInternalFromMaterialIndex = false;
+
 	// TODO: add primtive custom data
 };
 
@@ -400,7 +412,13 @@ public:
 	void ReindexMaterialSections();
 
 	/** appends the standard materials to this UObject */
-	void InitializeMaterials(bool bHasInternalMaterials = true);
+	void InitializeMaterials(bool bHasLegacyInternalMaterialsPairs = false);
+
+	/** Add a material to the materials array and update the selected bone material to be at the end of the array */
+	int32 AddNewMaterialSlot(bool bCopyLastMaterial = true);
+
+	/** Remove a material from the materials array, keeping the selected bone material at the end of the array. Returns false if materials could not be removed (e.g. because there were too few). */
+	bool RemoveLastMaterialSlot();
 
 
 	/** Returns true if there is anything to render */
