Fixing bad merge.

diff --git a/Engine/Plugins/Experimental/ChaosEditor/Source/FractureEditor/Private/FractureToolUV.cpp b/Engine/Plugins/Experimental/ChaosEditor/Source/FractureEditor/Private/FractureToolUV.cpp
--- a/Engine/Plugins/Experimental/ChaosEditor/Source/FractureEditor/Private/FractureToolUV.cpp
+++ b/Engine/Plugins/Experimental/ChaosEditor/Source/FractureEditor/Private/FractureToolUV.cpp
@@ -238,8 +238,8 @@ void UFractureToolAutoUV::BoxProjectUVs()
 			*GeometryCollectionComponent->GetRestCollection()->GetGeometryCollection(),
 			(FVector3d)AutoUVSettings->ProjectionScale,
 			TargetFaces,
-			!TargetsMaterialIDs(AutoUVSettings->TargetFaces) ? EmptyMaterialIDs : AutoUVSettings->MaterialIDs)
-			(FVector2f)AutoU			(FVector2f)AutoUVSettings->ProjectionUVOffset,
+			!TargetsMaterialIDs(AutoUVSettings->TargetFaces) ? EmptyMaterialIDs : AutoUVSettings->MaterialIDs,
+			(FVector2f)AutoUVSettings->ProjectionUVOffset,
 			AutoUVSettings->bAutoFitToBounds, AutoUVSettings->bCenterAtPivot, AutoUVSettings->bUniformProjectionScale);
 
 		GeometryCollectionComponent->MarkRenderDynamicDataDirty();
