non-unity fix

diff --git a/Engine/Plugins/Experimental/ChaosEditor/Source/FractureEditor/Private/FractureToolMaterials.cpp b/Engine/Plugins/Experimental/ChaosEditor/Source/FractureEditor/Private/FractureToolMaterials.cpp
--- a/Engine/Plugins/Experimental/ChaosEditor/Source/FractureEditor/Private/FractureToolMaterials.cpp
+++ b/Engine/Plugins/Experimental/ChaosEditor/Source/FractureEditor/Private/FractureToolMaterials.cpp
@@ -8,6 +8,7 @@
 #include "GeometryCollection/GeometryCollectionComponent.h"
 
 #include "FractureEngineMaterials.h"
+#include "ScopedTransaction.h"
 
 #include UE_INLINE_GENERATED_CPP_BY_NAME(FractureToolMaterials)
 
