fix mismatched tooltip on fracture split components option
cis Issue 262017

diff --git a/Engine/Plugins/Experimental/ChaosEditor/Source/FractureEditor/Private/FractureToolGenerators.cpp b/Engine/Plugins/Experimental/ChaosEditor/Source/FractureEditor/Private/FractureToolGenerators.cpp
--- a/Engine/Plugins/Experimental/ChaosEditor/Source/FractureEditor/Private/FractureToolGenerators.cpp
+++ b/Engine/Plugins/Experimental/ChaosEditor/Source/FractureEditor/Private/FractureToolGenerators.cpp
@@ -307,7 +307,7 @@ void SCreateGeometryCollectionFromObject::Construct(const FArguments& InArgs, TS
 					[
 						SNew(STextBlock)
 						.Text(LOCTEXT("SplitComponentsLabel", "Split Components"))
-						.ToolTipText(LOCTEXT("SplitComponentsToolTip", "If checked, triangles that are not topologically connected will be assigned separate bones on import."))
+						.ToolTipText(LOCTEXT("SplitComponentsToolTip", "If checked, triangles that are not topologically connected will be assigned separate bones on import. Does not apply if creating from a Geometry Collection source."))
 					]
 
 				+ SHorizontalBox::Slot()
