fracture mode rendering changes: rebuild render data rather than marking the component dirty for UVs and material updates, and change 'check' to 'ensure' for invalid exploded view vectors

diff --git a/Engine/Plugins/Experimental/ChaosEditor/Source/FractureEditor/Private/FractureToolMaterials.cpp b/Engine/Plugins/Experimental/ChaosEditor/Source/FractureEditor/Private/FractureToolMaterials.cpp
--- a/Engine/Plugins/Experimental/ChaosEditor/Source/FractureEditor/Private/FractureToolMaterials.cpp
+++ b/Engine/Plugins/Experimental/ChaosEditor/Source/FractureEditor/Private/FractureToolMaterials.cpp
@@ -37,8 +37,7 @@ void UFractureToolMaterials::RemoveMaterialSlot()
 		UGeometryCollection* Collection = Edit.GetRestCollection();
 		if (Collection->RemoveLastMaterialSlot())
 		{
-			GeometryCollectionComponent->MarkRenderDynamicDataDirty();
-			GeometryCollectionComponent->MarkRenderStateDirty();
+			Collection->RebuildRenderData();
 		}
 	}
 	UpdateActiveMaterialsList();
@@ -61,8 +60,7 @@ void UFractureToolMaterials::AddMaterialSlot()
 			GeometryCollectionComponent->SetMaterial(NewSlotIdx, GeometryCollectionComponent->GetMaterial(NewSlotIdx - 1));
 		}
 
-		GeometryCollectionComponent->MarkRenderDynamicDataDirty();
-		GeometryCollectionComponent->MarkRenderStateDirty();
+		Collection->RebuildRenderData();
 	}
 	UpdateActiveMaterialsList();
 }
diff --git a/Engine/Plugins/Experimental/ChaosEditor/Source/FractureEditor/Private/FractureToolUV.cpp b/Engine/Plugins/Experimental/ChaosEditor/Source/FractureEditor/Private/FractureToolUV.cpp
--- a/Engine/Plugins/Experimental/ChaosEditor/Source/FractureEditor/Private/FractureToolUV.cpp
+++ b/Engine/Plugins/Experimental/ChaosEditor/Source/FractureEditor/Private/FractureToolUV.cpp
@@ -196,8 +196,7 @@ void UFractureToolAutoUV::UpdateUVChannels(int32 TargetNumUVChannels)
 						}
 					}
 				}
-				GeometryCollectionComponent->MarkRenderDynamicDataDirty();
-				GeometryCollectionComponent->MarkRenderStateDirty();
+				Edit.GetRestCollection()->RebuildRenderData();
 			}
 		}
 		MinUVChannels = TargetNumUVChannels;
@@ -242,8 +241,7 @@ void UFractureToolAutoUV::BoxProjectUVs()
 			(FVector2f)AutoUVSettings->ProjectionUVOffset,
 			AutoUVSettings->bAutoFitToBounds, AutoUVSettings->bCenterAtPivot, AutoUVSettings->bUniformProjectionScale);
 
-		GeometryCollectionComponent->MarkRenderDynamicDataDirty();
-		GeometryCollectionComponent->MarkRenderStateDirty();
+		Edit.GetRestCollection()->RebuildRenderData();
 	}
 }
 
@@ -458,8 +456,7 @@ void UFractureToolAutoUV::LayoutUVs()
 		Edit.GetRestCollection()->Modify();
 		LayoutUVsForComponent(GeometryCollectionComponent);
 
-		GeometryCollectionComponent->MarkRenderDynamicDataDirty();
-		GeometryCollectionComponent->MarkRenderStateDirty();
+		Edit.GetRestCollection()->RebuildRenderData();
 	}
 }
 
diff --git a/Engine/Source/Runtime/Experimental/GeometryCollectionEngine/Private/GeometryCollection/GeometryCollectionComponent.cpp b/Engine/Source/Runtime/Experimental/GeometryCollectionEngine/Private/GeometryCollection/GeometryCollectionComponent.cpp
--- a/Engine/Source/Runtime/Experimental/GeometryCollectionEngine/Private/GeometryCollection/GeometryCollectionComponent.cpp
+++ b/Engine/Source/Runtime/Experimental/GeometryCollectionEngine/Private/GeometryCollection/GeometryCollectionComponent.cpp
@@ -468,7 +468,10 @@ namespace
 			{
 				const TManagedArray<FVector3f>& ExplodedVectors = GeometryCollection.GetAttribute<FVector3f>("ExplodedVector", FGeometryCollection::TransformGroup);
 
-				check(NumMatrices == ExplodedVectors.Num());
+				if (!ensure(NumMatrices == ExplodedVectors.Num()))
+				{
+					return;
+				}
 				for (int32 TransformIndex = 0; TransformIndex < ExplodedVectors.Num(); ++TransformIndex)
 				{
 					GlobalMatricesIn[TransformIndex] = GlobalMatricesIn[TransformIndex].ConcatTranslation((FVector)ExplodedVectors[TransformIndex]);
