add icon for material tool in fracture mode (copied over from modeling tools)

diff --git a/Engine/Plugins/Experimental/ChaosEditor/Source/FractureEditor/Private/FractureEditorStyle.cpp b/Engine/Plugins/Experimental/ChaosEditor/Source/FractureEditor/Private/FractureEditorStyle.cpp
--- a/Engine/Plugins/Experimental/ChaosEditor/Source/FractureEditor/Private/FractureEditorStyle.cpp
+++ b/Engine/Plugins/Experimental/ChaosEditor/Source/FractureEditor/Private/FractureEditorStyle.cpp
@@ -57,6 +57,7 @@ FFractureEditorStyle::FFractureEditorStyle()
 	Set("FractureEditor.FlushEmbeddedGeometry",new IMAGE_BRUSH_SVG("FractureFlush", IconSize));
 
 	Set("FractureEditor.AutoUV",               new IMAGE_BRUSH_SVG("FractureAutoUV", IconSize));
+	Set("FractureEditor.Material",             new IMAGE_BRUSH("EditMats_40x", IconSize));
 
 	Set("FractureEditor.DeleteBranch",         new IMAGE_BRUSH_SVG("FracturePrune", IconSize));
 	Set("FractureEditor.Hide",				   new IMAGE_BRUSH_SVG("FractureHide", IconSize));
