fix 'explode amount' slider not updating correctly when selected geometry collection is inside a blueprint

diff --git a/Engine/Plugins/Experimental/ChaosEditor/Source/FractureEditor/Private/FractureEditorModeToolkit.cpp b/Engine/Plugins/Experimental/ChaosEditor/Source/FractureEditor/Private/FractureEditorModeToolkit.cpp
--- a/Engine/Plugins/Experimental/ChaosEditor/Source/FractureEditor/Private/FractureEditorModeToolkit.cpp
+++ b/Engine/Plugins/Experimental/ChaosEditor/Source/FractureEditor/Private/FractureEditorModeToolkit.cpp
@@ -1099,16 +1099,12 @@ void FFractureEditorModeToolkit::OnExplodedViewValueChanged()
 		Actor->GetComponents(Components);
 		for (UPrimitiveComponent* PrimitiveComponent : Components)
 		{
-			AGeometryCollectionActor* GeometryCollectionActor = Cast<AGeometryCollectionActor>(Actor);
-			if(GeometryCollectionActor)
+			if (UGeometryCollectionComponent* GeometryCollectionComponent = Cast<UGeometryCollectionComponent>(PrimitiveComponent))
 			{
-				if (UGeometryCollectionComponent* GeometryCollectionComponent = Cast<UGeometryCollectionComponent>(PrimitiveComponent))
-				{
 
-					UpdateExplodedVectors(GeometryCollectionComponent);
+				UpdateExplodedVectors(GeometryCollectionComponent);
 
-					GeometryCollectionComponent->MarkRenderStateDirty();
-				}	
+				GeometryCollectionComponent->MarkRenderStateDirty();
 			}
 		}
 	}
