Update fracture mode materials tool to show and update the rest collection materials

diff --git a/Engine/Plugins/Experimental/ChaosEditor/Source/FractureEditor/Private/FractureToolMaterials.cpp b/Engine/Plugins/Experimental/ChaosEditor/Source/FractureEditor/Private/FractureToolMaterials.cpp
--- a/Engine/Plugins/Experimental/ChaosEditor/Source/FractureEditor/Private/FractureToolMaterials.cpp
+++ b/Engine/Plugins/Experimental/ChaosEditor/Source/FractureEditor/Private/FractureToolMaterials.cpp
@@ -26,6 +26,47 @@ void UFractureMaterialsSettings::RemoveMaterialSlot()
 	MaterialsTool->RemoveMaterialSlot();
 }
 
+void UFractureMaterialsSettings::UseAssetMaterialsOnComponents()
+{
+	UFractureToolMaterials* MaterialsTool = Cast<UFractureToolMaterials>(OwnerTool.Get());
+	MaterialsTool->ClearMaterialOverridesOnComponents(bOnlySelectedComponents);
+}
+
+void UFractureToolMaterials::ClearMaterialOverridesOnComponents(bool bOnlySelectedComponents)
+{
+	if (!ActiveSelectedComponent.IsValid())
+	{
+		return;
+	}
+
+	// update materials on all components using this GeometryCollection
+	if (!bOnlySelectedComponents)
+	{
+		for (FThreadSafeObjectIterator Iter(UGeometryCollectionComponent::StaticClass()); Iter; ++Iter)
+		{
+			UGeometryCollectionComponent* GCComponent = Cast<UGeometryCollectionComponent>(*Iter);
+			if (GCComponent->GetRestCollection() == ActiveSelectedComponent->GetRestCollection())
+			{
+				GCComponent->EmptyOverrideMaterials();
+				GCComponent->MarkRenderStateDirty();
+			}
+		}
+	}
+	else
+	{
+		TSet<UGeometryCollectionComponent*> GeomCompSelection;
+		GetSelectedGeometryCollectionComponents(GeomCompSelection);
+		for (UGeometryCollectionComponent* Component : GeomCompSelection)
+		{
+			if (Component->GetRestCollection() == ActiveSelectedComponent->GetRestCollection())
+			{
+				Component->EmptyOverrideMaterials();
+				Component->MarkRenderStateDirty();
+			}
+		}
+	}
+}
+
 void UFractureToolMaterials::RemoveMaterialSlot()
 {
 	TSet<UGeometryCollectionComponent*> GeomCompSelection;
@@ -40,7 +81,7 @@ void UFractureToolMaterials::RemoveMaterialSlot()
 			Collection->RebuildRenderData();
 		}
 	}
-	UpdateActiveMaterialsList();
+	UpdateActiveMaterialsInfo();
 }
 
 void UFractureToolMaterials::AddMaterialSlot()
@@ -54,15 +95,10 @@ void UFractureToolMaterials::AddMaterialSlot()
 		UGeometryCollection* Collection = Edit.GetRestCollection();
 
 		int32 NewSlotIdx = Collection->AddNewMaterialSlot();
-		if (NewSlotIdx > 0)
-		{
-			// copy an adjacent material into the new slot on the component as well
-			GeometryCollectionComponent->SetMaterial(NewSlotIdx, GeometryCollectionComponent->GetMaterial(NewSlotIdx - 1));
-		}
 
 		Collection->RebuildRenderData();
 	}
-	UpdateActiveMaterialsList();
+	UpdateActiveMaterialsInfo();
 }
 
 UFractureToolMaterials::UFractureToolMaterials(const FObjectInitializer& ObjInit)
@@ -94,7 +130,7 @@ FText UFractureToolMaterials::GetTooltipText() const
 
 FSlateIcon UFractureToolMaterials::GetToolIcon() const
 {
-	return FSlateIcon("FractureEditorStyle", "FractureEditor.ToMesh");
+	return FSlateIcon("FractureEditorStyle", "FractureEditor.Material");
 }
 
 void UFractureToolMaterials::RegisterUICommand(FFractureEditorCommands* BindingContext)
@@ -114,13 +150,28 @@ void UFractureToolMaterials::FractureContextChanged()
 {
 }
 
+void UFractureToolMaterials::SelectedBonesChanged()
+{
+	UpdateActiveMaterialsInfo();
+}
+
 void UFractureToolMaterials::Render(const FSceneView* View, FViewport* Viewport, FPrimitiveDrawInterface* PDI)
 {
 }
 
 void UFractureToolMaterials::PostEditChangeChainProperty(struct FPropertyChangedChainEvent& PropertyChangedEvent)
 {
-	// update any cached data 
+	Super::PostEditChangeChainProperty(PropertyChangedEvent);
+	if (PropertyChangedEvent.GetMemberPropertyName() == GET_MEMBER_NAME_STRING_CHECKED(UFractureMaterialsSettings, Materials))
+	{
+		if (ActiveSelectedComponent.IsValid())
+		{
+			FGeometryCollectionEdit GCEdit = ActiveSelectedComponent->EditRestCollection(GeometryCollection::EEditUpdate::Rest);
+			GCEdit.GetRestCollection()->Materials = MaterialsSettings->Materials;
+			GCEdit.GetRestCollection()->RebuildRenderData();
+			UpdateActiveMaterialsInfo();
+		}
+	}
 }
 
 TArray<FFractureToolContext> UFractureToolMaterials::GetFractureToolContexts() const
@@ -143,6 +194,49 @@ TArray<FFractureToolContext> UFractureToolMaterials::GetFractureToolContexts() c
 }
 
 
+void UFractureToolMaterials::UpdateActiveMaterialsInfo()
+{
+	// Choose a single asset to update for the asset materials options
+	TSet<UGeometryCollectionComponent*> GeomCompSelection;
+	GetSelectedGeometryCollectionComponents(GeomCompSelection, true);
+	const UGeometryCollection* RestCollection = nullptr;
+	UGeometryCollectionComponent* FoundComponent = nullptr;
+	for (UGeometryCollectionComponent* Component : GeomCompSelection)
+	{
+		FoundComponent = Component;
+		RestCollection = Component->GetRestCollection();
+		if (RestCollection)
+		{
+			break;
+		}
+	}
+	bool bHaveTargetCollectionUpdate = false;
+	if (RestCollection)
+	{
+		ActiveSelectedComponent = FoundComponent;
+		bHaveTargetCollectionUpdate = true;
+		int32 NumMaterials = RestCollection->Materials.Num();
+		MaterialsSettings->Materials = RestCollection->Materials;
+		RestCollection->GetName(MaterialsSettings->EditingCollection);
+	}
+	else
+	{
+		bHaveTargetCollectionUpdate = false;
+		MaterialsSettings->EditingCollection = LOCTEXT("NoActiveGeometryCollection", "None").ToString();
+		MaterialsSettings->Materials.Empty();
+		ActiveSelectedComponent = nullptr;
+	}
+	if (MaterialsSettings->bHaveTargetCollection != bHaveTargetCollectionUpdate)
+	{
+		MaterialsSettings->bHaveTargetCollection = bHaveTargetCollectionUpdate;
+		NotifyOfPropertyChangeByTool(MaterialsSettings);
+	}
+	
+	// Set active material list from full selection, since the material assignment applies to all selected collections
+	MaterialsSettings->UpdateActiveMaterialNames(GetSelectedComponentMaterialNames(false, false));
+}
+
+
 int32 UFractureToolMaterials::ExecuteFracture(const FFractureToolContext& FractureContext)
 {
 	if (FractureContext.IsValid())
@@ -164,7 +258,7 @@ int32 UFractureToolMaterials::ExecuteFracture(const FFractureToolContext& Fractu
 			(MaterialsSettings->ToFaces == EMaterialAssignmentTargets::OnlyInternalFaces) ? FFractureEngineMaterials::ETargetFaces::InternalFaces :
 			FFractureEngineMaterials::ETargetFaces::ExternalFaces;
 
-		if (MaterialsSettings->bOnlySelected)
+		if (MaterialsSettings->bOnlySelectedBones)
 		{
 			FFractureEngineMaterials::SetMaterial(Collection, FractureContext.GetSelection(), TargetFaces, MatID);
 		}
diff --git a/Engine/Plugins/Experimental/ChaosEditor/Source/FractureEditor/Private/FractureToolMaterials.h b/Engine/Plugins/Experimental/ChaosEditor/Source/FractureEditor/Private/FractureToolMaterials.h
--- a/Engine/Plugins/Experimental/ChaosEditor/Source/FractureEditor/Private/FractureToolMaterials.h
+++ b/Engine/Plugins/Experimental/ChaosEditor/Source/FractureEditor/Private/FractureToolMaterials.h
@@ -17,7 +17,7 @@ enum class EMaterialAssignmentTargets
 	AllFaces
 };
 
-/** Settings related to geometry collection -> static mesh conversion **/
+/** Settings related to setting materials on a Geometry Collection **/
 UCLASS(config = EditorPerProjectUserSettings)
 class UFractureMaterialsSettings : public UFractureToolSettings
 {
@@ -29,25 +29,42 @@ public:
 		: Super(ObjInit)
 	{}
 
+
+	/** Geometry Collection selected for editing */
+	UPROPERTY(VisibleAnywhere, Category = AssetMaterials, meta = (NoResetToDefault))
+	FString EditingCollection;
+	
+	/** Materials on the selected Geometry Collection's underlying asset (the Rest Collection). */
+	UPROPERTY(EditAnywhere, Category = AssetMaterials, EditFixedSize, meta = (DisplayName = "Asset Materials", NoResetToDefault, EditCondition = "bHaveTargetCollection", HideEditConditionToggle))
+	TArray<TObjectPtr<UMaterialInterface>> Materials;
+	
 	/** Add a new material slot to the selected geometry collections */
-	UFUNCTION(CallInEditor, Category = MaterialEditing, meta = (DisplayName = "Add Material Slot"))
+	UFUNCTION(CallInEditor, Category = AssetMaterials)
 	void AddMaterialSlot();
 
-	/** Remove the last material slot from the selected geometry collections. (Will not remove the final material.) */
-	UFUNCTION(CallInEditor, Category = MaterialEditing, meta = (DisplayName = "Remove Material Slot"))
+	/** Remove the last material slot from the selected Geometry Collections. (Will not remove the final material.) */
+	UFUNCTION(CallInEditor, Category = AssetMaterials)
 	void RemoveMaterialSlot();
 
+	/** Clear material overrides on components using this asset, so the asset materials are used */
+	UFUNCTION(CallInEditor, Category = ComponentMaterials)
+	void UseAssetMaterialsOnComponents();
+
+	/** Whether 'Use Asset Materials On Components' should only update the selected components, or update all components using this asset */
+	UPROPERTY(EditAnywhere, Category = ComponentMaterials)
+	bool bOnlySelectedComponents = false;
+
 	/** Material to assign to selected faces */
-	UPROPERTY(EditAnywhere, Category = MaterialEditing, meta = (TransientToolProperty, DisplayName = "Assign Material", GetOptions = GetMaterialNamesFunc, NoResetToDefault))
+	UPROPERTY(EditAnywhere, Category = AssignMaterials, meta = (TransientToolProperty, DisplayName = "Assign Material", GetOptions = GetMaterialNamesFunc, NoResetToDefault))
 	FString AssignMaterial;
 
 	/** Which subset of faces to update materials assignments on, for the selected geometry */
-	UPROPERTY(EditAnywhere, Category = MaterialEditing)
+	UPROPERTY(EditAnywhere, Category = AssignMaterials)
 	EMaterialAssignmentTargets ToFaces = EMaterialAssignmentTargets::OnlyInternalFaces;
 
 	/** Whether to only assign materials for faces in the selected bones, or the whole geometry collection */
-	UPROPERTY(EditAnywhere, Category = MaterialEditing)
-	bool bOnlySelected = true;
+	UPROPERTY(EditAnywhere, Category = AssignMaterials)
+	bool bOnlySelectedBones = true;
 
 	UFUNCTION()
 	const TArray<FString>& GetMaterialNamesFunc() { return AssignMaterialNamesList; }
@@ -66,6 +83,9 @@ public:
 		return AssignMaterialNamesList.Find(AssignMaterial);
 	}
 
+	UPROPERTY()
+	bool bHaveTargetCollection = false;
+
 private:
 
 	UPROPERTY(meta = (TransientToolProperty))
@@ -105,30 +125,34 @@ public:
 	}
 	virtual void PostEditChangeChainProperty(struct FPropertyChangedChainEvent& PropertyChangedEvent) override;
 
-	virtual void SelectedBonesChanged()
-	{
-		Super::SelectedBonesChanged();
-		UpdateActiveMaterialsList();
-	}
+	virtual void SelectedBonesChanged() override;
 
 	// Called when the modal tool is entered
-	virtual void Setup()
+	virtual void Setup() override
 	{
 		Super::Setup();
-		UpdateActiveMaterialsList();
+		UpdateActiveMaterialsInfo();
 	}
 
-	void UpdateActiveMaterialsList()
+	virtual void Shutdown()
 	{
-		MaterialsSettings->UpdateActiveMaterialNames(GetSelectedComponentMaterialNames(false, false));
+		Super::Shutdown();
+		ActiveSelectedComponent = nullptr;
 	}
 
+	void UpdateActiveMaterialsInfo();
+
 	virtual TArray<FFractureToolContext> GetFractureToolContexts() const override;
 
 	void RemoveMaterialSlot();
 	void AddMaterialSlot();
+	void ClearMaterialOverridesOnComponents(bool bOnlySelected = false);
 
 protected:
 	UPROPERTY(EditAnywhere, Category = Slicing)
 	TObjectPtr<UFractureMaterialsSettings> MaterialsSettings;
+
+	// Track the geometry collection asset last used to populate the materials list, if any
+	UPROPERTY()
+	TWeakObjectPtr<UGeometryCollectionComponent> ActiveSelectedComponent = nullptr;
 };
diff --git a/Engine/Plugins/Experimental/ChaosEditor/Source/FractureEditor/Public/FractureTool.cpp b/Engine/Plugins/Experimental/ChaosEditor/Source/FractureEditor/Public/FractureTool.cpp
--- a/Engine/Plugins/Experimental/ChaosEditor/Source/FractureEditor/Public/FractureTool.cpp
+++ b/Engine/Plugins/Experimental/ChaosEditor/Source/FractureEditor/Public/FractureTool.cpp
@@ -161,7 +161,7 @@ void UFractureActionTool::AddAdditionalAttributesIfRequired(UGeometryCollection*
 	}
 }
 
-void UFractureActionTool::GetSelectedGeometryCollectionComponents(TSet<UGeometryCollectionComponent*>& GeomCompSelection)
+void UFractureActionTool::GetSelectedGeometryCollectionComponents(TSet<UGeometryCollectionComponent*>& GeomCompSelection, bool bFilterForUniqueRestCollections)
 {
 	USelection* SelectionSet = GEditor->GetSelectedActors();
 	TArray<AActor*> SelectedActors;
@@ -170,11 +170,28 @@ void UFractureActionTool::GetSelectedGeometryCollectionComponents(TSet<UGeometry
 
 	GeomCompSelection.Empty(SelectionSet->Num());
 
+	TSet<const UGeometryCollection*> AddedRestCollections;
+
 	for (AActor* Actor : SelectedActors)
 	{
 		TInlineComponentArray<UGeometryCollectionComponent*> GeometryCollectionComponents;
 		Actor->GetComponents(GeometryCollectionComponents);
-		GeomCompSelection.Append(GeometryCollectionComponents);
+		if (bFilterForUniqueRestCollections)
+		{
+			for (UGeometryCollectionComponent* Component : GeometryCollectionComponents)
+			{
+				bool bWasInSet = false;
+				AddedRestCollections.FindOrAdd(Component->GetRestCollection(), &bWasInSet);
+				if (!bWasInSet)
+				{
+					GeomCompSelection.Add(Component);
+				}
+			}
+		}
+		else
+		{
+			GeomCompSelection.Append(GeometryCollectionComponents);
+		}
 	}
 }
 
@@ -215,7 +232,9 @@ TArray<FString> UFractureActionTool::GetSelectedComponentMaterialNames(bool bInc
 				{
 					continue;
 				}
-				MaterialNames.Add(FString::Printf(TEXT("[%d] %s"), MatIdx, *Component->GetMaterial(MatIdx)->GetName()));
+				UMaterialInterface* Material = Component->GetMaterial(MatIdx);
+				FString MaterialName = Material ? Material->GetName() : LOCTEXT("NoMaterialName", "None").ToString();
+				MaterialNames.Add(FString::Printf(TEXT("[%d] %s"), MatIdx, *MaterialName));
 			}
 		}
 	}
diff --git a/Engine/Plugins/Experimental/ChaosEditor/Source/FractureEditor/Public/FractureTool.h b/Engine/Plugins/Experimental/ChaosEditor/Source/FractureEditor/Public/FractureTool.h
--- a/Engine/Plugins/Experimental/ChaosEditor/Source/FractureEditor/Public/FractureTool.h
+++ b/Engine/Plugins/Experimental/ChaosEditor/Source/FractureEditor/Public/FractureTool.h
@@ -78,7 +78,13 @@ public:
 	};
 
 public:
-	static void GetSelectedGeometryCollectionComponents(TSet<UGeometryCollectionComponent*>& GeomCompSelection);
+
+	/**
+	 * Get the Geometry Collection components in the current selection
+	 * @param GeomCompSelection					the set to fill with Geometry Collection components
+	 * @param bFilterForUniqueRestCollections	if true, and selection contains multiple components backed by the same Rest Collection asset, the selection will be filtered to only contain one component per asset
+	 */
+	static void GetSelectedGeometryCollectionComponents(TSet<UGeometryCollectionComponent*>& GeomCompSelection, bool bFilterForUniqueRestCollections = false);
 	static TArray<FString> GetSelectedComponentMaterialNames(bool bIncludeDefault, bool bUseFullNamesIfPossible = true);
 
 	
