add Rest Collection to the fracture mode view settings

diff --git a/Engine/Plugins/Experimental/ChaosEditor/Source/FractureEditor/Private/FractureEditorModeToolkit.cpp b/Engine/Plugins/Experimental/ChaosEditor/Source/FractureEditor/Private/FractureEditorModeToolkit.cpp
--- a/Engine/Plugins/Experimental/ChaosEditor/Source/FractureEditor/Private/FractureEditorModeToolkit.cpp
+++ b/Engine/Plugins/Experimental/ChaosEditor/Source/FractureEditor/Private/FractureEditorModeToolkit.cpp
@@ -1478,6 +1478,15 @@ void FFractureEditorModeToolkit::SetOutlinerComponents(const TArray<UGeometryCol
 		StatisticsView->SetStatistics(Stats);
 	}
 
+	// update view rest collection
+	UFractureSettings* FractureSettings = GetMutableDefault<UFractureSettings>();
+	FractureSettings->RestCollection = nullptr;
+	for (UGeometryCollectionComponent* Component : InNewComponents)
+	{
+		FractureSettings->RestCollection = Component->GetRestCollection();
+		break;
+	}
+
 	if (ActiveTool != nullptr)
 	{
 		ActiveTool->SelectedBonesChanged();
diff --git a/Engine/Plugins/Experimental/ChaosEditor/Source/FractureEditor/Public/FractureSettings.h b/Engine/Plugins/Experimental/ChaosEditor/Source/FractureEditor/Public/FractureSettings.h
--- a/Engine/Plugins/Experimental/ChaosEditor/Source/FractureEditor/Public/FractureSettings.h
+++ b/Engine/Plugins/Experimental/ChaosEditor/Source/FractureEditor/Public/FractureSettings.h
@@ -3,9 +3,12 @@
 #pragma once
 
 #include "UObject/Object.h"
+#include "UObject/WeakObjectPtrTemplates.h"
 
 #include "FractureSettings.generated.h"
 
+class UGeometryCollection;
+
 /** Settings specifically related to viewing fractured meshes **/
 UCLASS()
 class UFractureSettings: public UObject
@@ -27,4 +30,7 @@ public:
 	UPROPERTY(EditAnywhere, Category = ViewSettings, meta = (DisplayName = "Hide Unselected"))
 	bool bHideUnselected;
 
+	UPROPERTY(VisibleAnywhere, Category = ViewSettings)
+	TWeakObjectPtr<const UGeometryCollection> RestCollection = nullptr;
+
 };
\ No newline at end of file
