Fracture editor : improve connectivity tool visualization
- Centers are now displayed
- Add controls for line hickness, point size and colors

diff --git a/Engine/Plugins/Experimental/ChaosEditor/Source/FractureEditor/Private/FractureToolProximity.cpp b/Engine/Plugins/Experimental/ChaosEditor/Source/FractureEditor/Private/FractureToolProximity.cpp
--- a/Engine/Plugins/Experimental/ChaosEditor/Source/FractureEditor/Private/FractureToolProximity.cpp
+++ b/Engine/Plugins/Experimental/ChaosEditor/Source/FractureEditor/Private/FractureToolProximity.cpp
@@ -183,7 +183,13 @@ void UFractureToolProximity::Render(const FSceneView* View, FViewport* Viewport,
 			{
 				FVector P1 = WorldTransform.TransformPosition(Vis.GeoCenters[Edge.A] + GetExplodedOffset(Edge.A));
 				FVector P2 = WorldTransform.TransformPosition(Vis.GeoCenters[Edge.B] + GetExplodedOffset(Edge.B));
-				PDI->DrawLine(P1, P2, FLinearColor::Green, SDPG_Foreground, 0.0f, 0.001f);
+				PDI->DrawLine(P1, P2, ProximitySettings->LineColor, SDPG_Foreground, ProximitySettings->LineThickness, 0.001f);
+			}
+
+			for (int32 CenterIndex = 0; CenterIndex < Vis.GeoCenters.Num(); CenterIndex++)
+			{
+				FVector P = WorldTransform.TransformPosition(Vis.GeoCenters[CenterIndex] + GetExplodedOffset(CenterIndex));
+				PDI->DrawPoint(P, ProximitySettings->CenterColor, ProximitySettings->CenterSize, SDPG_Foreground);
 			}
 		}
 	}
diff --git a/Engine/Plugins/Experimental/ChaosEditor/Source/FractureEditor/Private/FractureToolProximity.h b/Engine/Plugins/Experimental/ChaosEditor/Source/FractureEditor/Private/FractureToolProximity.h
--- a/Engine/Plugins/Experimental/ChaosEditor/Source/FractureEditor/Private/FractureToolProximity.h
+++ b/Engine/Plugins/Experimental/ChaosEditor/Source/FractureEditor/Private/FractureToolProximity.h
@@ -49,6 +49,22 @@ public:
 	// Whether to only show the proximity graph connections for selected bones
 	UPROPERTY(EditAnywhere, Category = Visualization, meta = (EditCondition = "bShowProximity"))
 	bool bOnlyShowForSelected = false;
+
+	// Line thickness for connections
+	UPROPERTY(EditAnywhere, Category = Visualization, meta = (ClampMin = "0.0", ClampMax = "1.0"))
+	float LineThickness = 0.5;
+
+	// Line color for connections
+	UPROPERTY(EditAnywhere, Category = Visualization)
+	FColor LineColor = FColor::Yellow;
+
+	// Point size for centers
+	UPROPERTY(EditAnywhere, Category = Visualization, meta = (ClampMin = "0.0", ClampMax = "20.0"))
+	float CenterSize = 8.f;
+
+	// Point color for centers
+	UPROPERTY(EditAnywhere, Category = Visualization)
+	FColor CenterColor = FColor::Blue;
 };
 
 
