Add new fracture plugin project settings to control the default location for new rest collection assets.

Add 'quick settings' section to the Fracture Editor UI, to allow users to easily control the default folder for new assets and to jump to fracture project and editor settings.

Update fracture mode UI so that quick settings and modal tool accept/cancel buttons are always visible without scrolling.

diff --git a/Engine/Plugins/Experimental/ChaosEditor/Source/FractureEditor/FractureEditor.Build.cs b/Engine/Plugins/Experimental/ChaosEditor/Source/FractureEditor/FractureEditor.Build.cs
--- a/Engine/Plugins/Experimental/ChaosEditor/Source/FractureEditor/FractureEditor.Build.cs
+++ b/Engine/Plugins/Experimental/ChaosEditor/Source/FractureEditor/FractureEditor.Build.cs
@@ -45,6 +45,7 @@ public class FractureEditor : ModuleRules
 				"ToolWidgets",
 				"DeveloperSettings",
 				"FractureEngine",
+				"ContentBrowserData",
 				"GeometryFramework"
 
 				// ... add private dependencies that you statically link with here ...	
diff --git a/Engine/Plugins/Experimental/ChaosEditor/Source/FractureEditor/Private/FractureEditorModeToolkit.cpp b/Engine/Plugins/Experimental/ChaosEditor/Source/FractureEditor/Private/FractureEditorModeToolkit.cpp
--- a/Engine/Plugins/Experimental/ChaosEditor/Source/FractureEditor/Private/FractureEditorModeToolkit.cpp
+++ b/Engine/Plugins/Experimental/ChaosEditor/Source/FractureEditor/Private/FractureEditorModeToolkit.cpp
@@ -35,6 +35,7 @@
 #include "Internationalization/Text.h"
 
 #include "FractureTool.h"
+#include "FractureModeSettings.h"
 
 #include "GeometryCollection/GeometryCollectionActor.h"
 #include "GeometryCollection/GeometryCollection.h"
@@ -47,6 +48,7 @@
 #include "FractureEditorStyle.h"
 #include "ScopedTransaction.h"
 #include "Modules/ModuleManager.h"
+#include "ISettingsModule.h"
 
 #include "PlanarCut.h"
 #include "FractureToolAutoCluster.h" 
@@ -68,6 +70,7 @@
 #include "DetailCategoryBuilder.h"
 #include "DetailLayoutBuilder.h"
 #include "PropertyHandle.h"
+#include "Widgets/Input/STextComboBox.h"
 
 #include "LevelEditor.h"
 
@@ -199,8 +202,51 @@ FFractureEditorModeToolkit::~FFractureEditorModeToolkit()
 		auto& LevelEditorModule = FModuleManager::LoadModuleChecked<FLevelEditorModule>(TEXT("LevelEditor"));
 		LevelEditorModule.OnMapChanged().RemoveAll(this);
 	}
+	UFractureModeSettings* Settings = GetMutableDefault<UFractureModeSettings>();
+	Settings->OnModified.Remove(ProjectSettingsModifiedHandle);
 }
 
+void FFractureEditorModeToolkit::UpdateAssetLocationMode(TSharedPtr<FString> NewString)
+{
+	UFractureModeSettings* Settings = GetMutableDefault<UFractureModeSettings>();
+	EFractureModeNewAssetLocation NewAssetLocation = EFractureModeNewAssetLocation::SourceAssetFolder;
+	if (NewString == AssetLocationModes[0])
+	{
+		NewAssetLocation = EFractureModeNewAssetLocation::SourceAssetFolder;
+	}
+	else if (NewString == AssetLocationModes[1])
+	{
+		NewAssetLocation = EFractureModeNewAssetLocation::LastUsedFolder;
+	}
+	else if (NewString == AssetLocationModes[2])
+	{
+		NewAssetLocation = EFractureModeNewAssetLocation::ContentBrowserFolder;
+	}
+
+	Settings->NewAssetLocation = NewAssetLocation;
+	Settings->SaveConfig();
+}
+
+void FFractureEditorModeToolkit::UpdateAssetPanelFromSettings()
+{
+	const UFractureModeSettings* Settings = GetDefault<UFractureModeSettings>();
+
+	switch (Settings->NewAssetLocation)
+	{
+	case EFractureModeNewAssetLocation::ContentBrowserFolder:
+		AssetLocationMode->SetSelectedItem(AssetLocationModes[2]);
+		break;
+	case EFractureModeNewAssetLocation::LastUsedFolder:
+		AssetLocationMode->SetSelectedItem(AssetLocationModes[1]);
+		break;
+	case EFractureModeNewAssetLocation::SourceAssetFolder:
+	default:
+		AssetLocationMode->SetSelectedItem(AssetLocationModes[0]);
+		break;
+	}
+}
+
+
 void FFractureEditorModeToolkit::Init(const TSharedPtr<IToolkitHost>& InitToolkitHost, TWeakObjectPtr<UEdMode> InOwningMode)
 {
 	FFractureEditorModule& FractureModule = FModuleManager::GetModuleChecked<FFractureEditorModule>("FractureEditor");
@@ -246,13 +292,78 @@ void FFractureEditorModeToolkit::Init(const TSharedPtr<IToolkitHost>& InitToolki
 	EditModule.RegisterCustomClassLayout("OutlinerSettings", FOnGetDetailCustomizationInstance::CreateStatic(&FOutlinerSettingsCustomization::MakeInstance, this));
 	OutlinerDetailsView->SetObject(GetMutableDefault<UOutlinerSettings>());
 
+	AssetLocationModes.Reset();
+	AssetLocationModes.Add(MakeShared<FString>(TEXT("Source Asset Folder")));
+	AssetLocationModes.Add(MakeShared<FString>(TEXT("Last Used Folder")));
+	AssetLocationModes.Add(MakeShared<FString>(TEXT("Content Browser Folder")));
+	AssetLocationMode = SNew(STextComboBox)
+		.OptionsSource(&AssetLocationModes)
+		.OnSelectionChanged_Lambda([&](TSharedPtr<FString> String, ESelectInfo::Type) { UpdateAssetLocationMode(String); });
+
+	const TSharedPtr<SVerticalBox> Content = SNew(SVerticalBox)
+		+ SVerticalBox::Slot().HAlign(HAlign_Fill)
+		[
+			SNew(SHorizontalBox)
+				+ SHorizontalBox::Slot().Padding(0).HAlign(HAlign_Left).VAlign(VAlign_Center).FillWidth(2.f)
+				[
+					SNew(STextBlock)
+					.Text(LOCTEXT("AssetLocationLabel", "New Asset Location"))
+				]
+				+ SHorizontalBox::Slot().Padding(0).FillWidth(4.0f)
+				[
+					AssetLocationMode->AsShared()
+				]
+		];
+
+	TSharedPtr<SExpandableArea> AssetConfigPanel = SNew(SExpandableArea)
+		.HeaderPadding(FMargin(0.f))
+		.Padding(FMargin(8.f))
+		.BorderImage(FAppStyle::Get().GetBrush("DetailsView.CategoryTop"))
+		.AreaTitleFont(FAppStyle::Get().GetFontStyle("EditorModesPanel.CategoryFontStyle"))
+		.BodyContent()
+		[
+			Content->AsShared()
+		]
+		.HeaderContent()
+		[
+			SNew(SHorizontalBox)
+			+ SHorizontalBox::Slot().HAlign(HAlign_Left).VAlign(VAlign_Center).FillWidth(2.f)
+			[
+				SNew(STextBlock)
+				.Text(LOCTEXT("FractureSettingsPanelHeader", "Fracture Mode Quick Settings"))
+			]
+
+			+ SHorizontalBox::Slot().HAlign(HAlign_Right).VAlign(VAlign_Center).AutoWidth()
+			[
+				SNew(SComboButton)
+				.HasDownArrow(false)
+				.MenuPlacement(EMenuPlacement::MenuPlacement_MenuRight)
+				.ComboButtonStyle(FAppStyle::Get(), "SimpleComboButton")
+				.OnGetMenuContent(FOnGetContent::CreateLambda([this]()
+				{
+					return MakeMenu_FractureModeConfigSettings();
+				}))
+				.ContentPadding(FMargin(3.0f, 1.0f))
+				.ButtonContent()
+				[
+					SNew(SImage)
+					.Image(FFractureEditorStyle::Get().GetBrush("FractureEditor.DefaultSettings"))
+					.ColorAndOpacity(FSlateColor::UseForeground())
+				]
+			]
+
+		];
+
+	const TSharedPtr<SVerticalBox> ToolkitWidgetVBox = SNew(SVerticalBox);
+
 	float Padding = 4.0f;
 	FMargin MorePadding = FMargin(10.0f, 2.0f);
-	SAssignNew(ToolkitWidget, SBox)
-	[
-		SNew(SVerticalBox)
-
-		+SVerticalBox::Slot()
+	SAssignNew(ToolkitWidget, SBorder).HAlign(HAlign_Fill)
+		.Padding(4)
+		[
+			ToolkitWidgetVBox->AsShared()
+		];
+	ToolkitWidgetVBox->AddSlot().HAlign(HAlign_Fill).FillHeight(1.f)
 		[
 
 			SNew(SSplitter)
@@ -304,12 +415,19 @@ void FFractureEditorModeToolkit::Init(const TSharedPtr<IToolkitHost>& InitToolki
 					]
 				]
 			]
-		]
-	];
-
+		];
+	ToolkitWidgetVBox->AddSlot().AutoHeight().HAlign(HAlign_Fill).VAlign(VAlign_Bottom).Padding(0)
+		[
+			AssetConfigPanel->AsShared()
+		];
 
 
+	// register callback
+	UFractureModeSettings* FractureModeSettings = GetMutableDefault<UFractureModeSettings>();
+	ProjectSettingsModifiedHandle = FractureModeSettings->OnModified.AddLambda([this](UObject*, FProperty*) { OnProjectSettingsModified(); });
 
+	// initialize combos
+	UpdateAssetPanelFromSettings();
 
 	// Bind Chaos Commands;
 	BindCommands();
@@ -318,6 +436,66 @@ void FFractureEditorModeToolkit::Init(const TSharedPtr<IToolkitHost>& InitToolki
 
 }
 
+void FFractureEditorModeToolkit::OnProjectSettingsModified()
+{
+	UpdateAssetPanelFromSettings();
+}
+
+namespace
+{
+
+	void MakeFractureQuickSettings(FMenuBuilder& MenuBuilder)
+	{
+		const FUIAction OpenFractureModeProjectSettings(
+			FExecuteAction::CreateLambda([]
+				{
+					if (ISettingsModule* SettingsModule = FModuleManager::GetModulePtr<ISettingsModule>("Settings"))
+					{
+						SettingsModule->ShowViewer("Project", "Plugins", "FractureMode");
+					}
+				}), FCanExecuteAction(), FIsActionChecked());
+		MenuBuilder.AddMenuEntry(LOCTEXT("FractureModeProjectSettings", "Project Settings"),
+			LOCTEXT("FractureModeProjectSettings_Tooltip", "Jump to the Project Settings for Fracture Mode. Project Settings are Project-specific."),
+			FSlateIcon(), OpenFractureModeProjectSettings, NAME_None, EUserInterfaceActionType::Button);
+
+		const FUIAction OpenFractureModeEditorSettings(
+			FExecuteAction::CreateLambda([]
+				{
+					if (ISettingsModule* SettingsModule = FModuleManager::GetModulePtr<ISettingsModule>("Settings"))
+					{
+						SettingsModule->ShowViewer("Editor", "Plugins", "FractureEditor");
+					}
+				}), FCanExecuteAction(), FIsActionChecked());
+		MenuBuilder.AddMenuEntry(LOCTEXT("FractureModeEditorSettings", "Editor Settings"),
+			LOCTEXT("FractureModeEditorSettings_Tooltip", "Jump to the Editor Settings for Fracture Mode. Editor Settings apply across all Projects."),
+			FSlateIcon(), OpenFractureModeEditorSettings, NAME_None, EUserInterfaceActionType::Button);
+	}
+}
+
+TSharedRef<SWidget> FFractureEditorModeToolkit::MakeMenu_FractureModeConfigSettings()
+{
+	FMenuBuilder MenuBuilder(true, TSharedPtr<FUICommandList>());
+
+	MenuBuilder.BeginSection("Section_Settings", LOCTEXT("Section_Settings", "Quick Settings"));
+	constexpr bool bQuickSettingsInSubMenu = false;
+	if (!bQuickSettingsInSubMenu)
+	{
+		MakeFractureQuickSettings(MenuBuilder);
+	}
+	else
+	{
+		MenuBuilder.AddSubMenu(
+			LOCTEXT("QuickSettingsSubMenu", "Jump To Settings"), LOCTEXT("QuickSettingsSubMenu_ToolTip", "Jump to sections of the Settings dialogs relevant to Fracture Mode"),
+			FNewMenuDelegate::CreateLambda([=](FMenuBuilder& SubMenuBuilder) {
+				MakeFractureQuickSettings(SubMenuBuilder);
+				}));
+	}
+	MenuBuilder.EndSection();
+
+	TSharedRef<SWidget> MenuWidget = MenuBuilder.MakeWidget();
+	return MenuWidget;
+}
+
 void FFractureEditorModeToolkit::RequestModeUITabs()
 {
 	FModeToolkit::RequestModeUITabs();
@@ -345,10 +523,28 @@ void FFractureEditorModeToolkit::RequestModeUITabs()
 	}
 }
 
+TSharedPtr<SWidget> FFractureEditorModeToolkit::GetInlineContent() const
+{
+	return SNew(SVerticalBox)
+		+ SVerticalBox::Slot()
+		.FillHeight(1.0f)
+		.VAlign(VAlign_Fill)
+		[
+			ToolkitWidget.ToSharedRef()
+		];
+}
+
 void FFractureEditorModeToolkit::InvokeUI()
 {
 	FModeToolkit::InvokeUI();
 
+	// (Note this logic is from ModelingToolsEditorModeToolkit.cpp, as the modeling tools had the same issue)
+	// FModeToolkit::UpdatePrimaryModePanel() wrapped our GetInlineContent() output in a SScrollBar widget,
+	// however this doesn't make sense as we want to dock panels to the "top" and "bottom" of our mode panel area,
+	// and the details panel in the middle has it's own scrollbar already. The SScrollBar is hardcoded as the content
+	// of FModeToolkit::InlineContentHolder so we can just replace it here
+	InlineContentHolder->SetContent(GetInlineContent().ToSharedRef());
+
 	if (TSharedPtr<FAssetEditorModeUILayer> ModeUILayerPtr = ModeUILayer.Pin())
 	{
 		TSharedPtr<FTabManager> TabManagerPtr = ModeUILayerPtr->GetTabManager();
diff --git a/Engine/Plugins/Experimental/ChaosEditor/Source/FractureEditor/Private/FractureEditorStyle.cpp b/Engine/Plugins/Experimental/ChaosEditor/Source/FractureEditor/Private/FractureEditorStyle.cpp
--- a/Engine/Plugins/Experimental/ChaosEditor/Source/FractureEditor/Private/FractureEditorStyle.cpp
+++ b/Engine/Plugins/Experimental/ChaosEditor/Source/FractureEditor/Private/FractureEditorStyle.cpp
@@ -20,6 +20,8 @@ FFractureEditorStyle::FFractureEditorStyle()
 	SetContentRoot(FPaths::EnginePluginsDir() / TEXT("Experimental/ChaosEditor/Content"));
 	SetCoreContentRoot(FPaths::EngineContentDir() / TEXT("Slate"));
 
+	Set("FractureEditor.DefaultSettings", new FSlateImageBrush(RootToCoreContentDir(TEXT("../Editor/Slate/Icons/GeneralTools/Settings_40x.png")), IconSize));
+
 	Set("FractureEditor.Slice",                new IMAGE_BRUSH_SVG("FractureSliceCut", IconSize));
 	Set("FractureEditor.Uniform",              new IMAGE_BRUSH_SVG("FractureUniformVoronoi", IconSize));
 	Set("FractureEditor.Radial",               new IMAGE_BRUSH_SVG("FractureRadialVoronoi", IconSize));
diff --git a/Engine/Plugins/Experimental/ChaosEditor/Source/FractureEditor/Private/FractureToolGenerators.cpp b/Engine/Plugins/Experimental/ChaosEditor/Source/FractureEditor/Private/FractureToolGenerators.cpp
--- a/Engine/Plugins/Experimental/ChaosEditor/Source/FractureEditor/Private/FractureToolGenerators.cpp
+++ b/Engine/Plugins/Experimental/ChaosEditor/Source/FractureEditor/Private/FractureToolGenerators.cpp
@@ -56,7 +56,16 @@
 #define LOCTEXT_NAMESPACE "FractureToolGenerators"
 
 
+namespace
+{
+	FString GetActiveContentBrowserFolderPath()
+	{
+		IContentBrowserSingleton& ContentBrowser = FModuleManager::LoadModuleChecked<FContentBrowserModule>("ContentBrowser").Get();
+		const FContentBrowserItemPath CurrentPath = ContentBrowser.GetCurrentPath();
+		return CurrentPath.HasInternalPath() ? CurrentPath.GetInternalPathString() : FString();
+	}
 
+}
 
 
 //////////////////////////////////////////////////////////////////////////
@@ -496,6 +505,79 @@ void UFractureToolGenerateAsset::Execute(TWeakPtr<FFractureEditorModeToolkit> In
 	// Note: Transaction for undo history is created after the user completes the UI dialog; see OnGenerateAssetPathChosen()
 }
 
+FString UFractureToolGenerateAsset::GetDefaultAssetPath(const TArray<AActor*>& Actors) const
+{
+	FString UseAssetPath;
+
+	const UFractureModeSettings* Settings = GetDefault<UFractureModeSettings>();
+	bool bFailedToFindContentBrowser = false, bFailedToFindLastUsedFolder = false;
+	if (Settings->NewAssetLocation == EFractureModeNewAssetLocation::ContentBrowserFolder)
+	{
+		FString ActiveContentFolder = GetActiveContentBrowserFolderPath();
+		if (ActiveContentFolder.IsEmpty())
+		{
+			bFailedToFindContentBrowser = true;
+		}
+		else
+		{
+			UseAssetPath = ActiveContentFolder;
+		}
+	}
+	if (bFailedToFindContentBrowser || Settings->NewAssetLocation == EFractureModeNewAssetLocation::LastUsedFolder)
+	{
+		if (AssetPath.IsEmpty())
+		{
+			bFailedToFindLastUsedFolder = true;
+		}
+		else
+		{
+			UseAssetPath = AssetPath;
+		}
+	}
+	if (bFailedToFindLastUsedFolder || Settings->NewAssetLocation == EFractureModeNewAssetLocation::SourceAssetFolder)
+	{
+		for (const AActor* Actor : Actors)
+		{
+			// Find the asset folder of any selected static mesh
+			TArray<UStaticMeshComponent*> StaticMeshComponents;
+			Actor->GetComponents(StaticMeshComponents, true);
+			const UStaticMesh* FoundMeshAsset = nullptr;
+			for (UStaticMeshComponent* StaticMeshComponent : StaticMeshComponents)
+			{
+				FoundMeshAsset = StaticMeshComponent ? StaticMeshComponent->GetStaticMesh() : nullptr;
+				if (FoundMeshAsset)
+				{
+					break;
+				}
+			}
+			if (FoundMeshAsset)
+			{
+				UseAssetPath = FoundMeshAsset->GetPathName();
+				break;
+			}
+			// Find the asset folder of any selected geometry collection
+			TArray<UGeometryCollectionComponent*> GeometryCollectionComponents;
+			Actor->GetComponents(GeometryCollectionComponents, true);
+			const UGeometryCollection* FoundRestAsset = nullptr;
+			for (UGeometryCollectionComponent* GCComp : GeometryCollectionComponents)
+			{
+				FoundRestAsset = GCComp ? GCComp->GetRestCollection() : nullptr;
+				if (FoundRestAsset)
+				{
+					break;
+				}
+			}
+			if (FoundRestAsset)
+			{
+				UseAssetPath = FoundRestAsset->GetPathName();
+				break;
+			}
+		}
+	}
+	return UseAssetPath;
+}
+
+
 void UFractureToolGenerateAsset::OpenGenerateAssetDialog(TArray<AActor*>& Actors)
 {
 	TSharedPtr<SWindow> PickAssetPathWindow;
@@ -505,6 +587,8 @@ void UFractureToolGenerateAsset::OpenGenerateAssetDialog(TArray<AActor*>& Actors
 		.ToolTipText(LOCTEXT("SelectPathTooltip", "Select the asset path for your new Geometry Collection"))
 		.ClientSize(FVector2D(500, 500));
 
+	AssetPath = GetDefaultAssetPath(Actors);
+
 	// NOTE - the parent window has to completely exist before this one does so the parent gets set properly.
 	// This is why we do not just put this in the Contents()[ ... ] of the Window above.
 	TSharedPtr<SCreateGeometryCollectionFromObject> CreateAssetDialog;
diff --git a/Engine/Plugins/Experimental/ChaosEditor/Source/FractureEditor/Private/FractureToolGenerators.h b/Engine/Plugins/Experimental/ChaosEditor/Source/FractureEditor/Private/FractureToolGenerators.h
--- a/Engine/Plugins/Experimental/ChaosEditor/Source/FractureEditor/Private/FractureToolGenerators.h
+++ b/Engine/Plugins/Experimental/ChaosEditor/Source/FractureEditor/Private/FractureToolGenerators.h
@@ -36,6 +36,8 @@ private:
 	static ULevel* GetSelectedLevel();
 	static AActor* AddActor(ULevel* InLevel, UClass* Class);
 
+	FString GetDefaultAssetPath(const TArray<AActor*>& Actors) const;
+
 private:
 	TWeakPtr<FFractureEditorModeToolkit> Toolkit;
 
diff --git a/Engine/Plugins/Experimental/ChaosEditor/Source/FractureEditor/Public/FractureEditorModeToolkit.h b/Engine/Plugins/Experimental/ChaosEditor/Source/FractureEditor/Public/FractureEditorModeToolkit.h
--- a/Engine/Plugins/Experimental/ChaosEditor/Source/FractureEditor/Public/FractureEditorModeToolkit.h
+++ b/Engine/Plugins/Experimental/ChaosEditor/Source/FractureEditor/Public/FractureEditorModeToolkit.h
@@ -100,7 +100,7 @@ public:
 	virtual FName GetToolkitFName() const override;
 	virtual FText GetBaseToolkitName() const override;
 	virtual class FEdMode* GetEditorMode() const override;
-	virtual TSharedPtr<class SWidget> GetInlineContent() const override { return ToolkitWidget; }
+	virtual TSharedPtr<class SWidget> GetInlineContent() const override;
 
 	/** FGCObject interface */
 	virtual void AddReferencedObjects(FReferenceCollector& Collector) override;
@@ -216,12 +216,19 @@ private:
 	void HandleMapChanged(UWorld* NewWorld, EMapChangeType MapChangeType);
 
 	FReply OnRefreshOutlinerButtonClicked();
+
+	TSharedRef<SWidget> MakeMenu_FractureModeConfigSettings();
+	void UpdateAssetLocationMode(TSharedPtr<FString> NewString);
+	void UpdateAssetPanelFromSettings();
+	void OnProjectSettingsModified();
 	
 private:
 	UFractureModalTool* ActiveTool;
 
 	// called when PIE is about to start, shuts down active tools
 	FDelegateHandle BeginPIEDelegateHandle;
+	// calls with the project settings are modified; used to keep the quick settings up to date
+	FDelegateHandle ProjectSettingsModifiedHandle;
 
 	TSharedPtr<IDetailsView> DetailsView;
 	TSharedPtr<IDetailsView> ViewSettingsDetailsView;
@@ -235,5 +242,7 @@ private:
 	TWeakPtr<SDockTab> StatisticsTab;
 	FMinorTabConfig StatisticsTabInfo;
 	TSharedPtr<SGeometryCollectionStatistics> StatisticsView;
+	TArray<TSharedPtr<FString>> AssetLocationModes;
+	TSharedPtr<STextComboBox> AssetLocationMode;
 
 };
diff --git a/Engine/Plugins/Experimental/ChaosEditor/Source/FractureEditor/Public/FractureModeSettings.h b/Engine/Plugins/Experimental/ChaosEditor/Source/FractureEditor/Public/FractureModeSettings.h
--- a/Engine/Plugins/Experimental/ChaosEditor/Source/FractureEditor/Public/FractureModeSettings.h
+++ b/Engine/Plugins/Experimental/ChaosEditor/Source/FractureEditor/Public/FractureModeSettings.h
@@ -12,6 +12,20 @@
 
 class FGeometryCollection;
 
+// Control the default asset folder presented when using the "New" tool to create a Geometry Collection in Fracture Mode
+UENUM()
+enum class EFractureModeNewAssetLocation
+{
+	/** Default to creating rest collections in the same folder as the source asset. */
+	SourceAssetFolder,
+
+	/** Default to creating rest collections in the last folder selected this session. If no folder was selected yet, use the Source Asset Folder. */
+	LastUsedFolder,
+
+	/** Default to creating reset collections in the currently-visible Asset Browser folder if available, otherwise use the Last Used Folder. */
+	ContentBrowserFolder
+};
+
 /**
  * Settings for the Fracture Editor Mode.
  */
@@ -33,6 +47,10 @@ public:
 
 public:
 
+	/** The default asset folder presented when using the "New" tool to create a Geometry Collection in Fracture Mode */
+	UPROPERTY(config, EditAnywhere, Category = "Fracture Mode|Asset Location")
+	EFractureModeNewAssetLocation NewAssetLocation = EFractureModeNewAssetLocation::SourceAssetFolder;
+
 	/** Default fraction of geometry volume by which a cluster's convex hull volume can exceed the actual geometry volume before instead using the hulls of the children.  0 means the convex volume cannot exceed the geometry volume; 1 means the convex volume is allowed to be 100% larger (2x) the geometry volume. */
 	UPROPERTY(config, EditAnywhere, Category = "Fracture Mode|Convex Generation Defaults", meta = (ClampMin = "0", DisplayName = "Allow Larger Hull Fraction"))
 	float ConvexCanExceedFraction = .5;
