protect against component having out of date selection by adding a Sanitize() before processing the selection in fracture editor's OnLevelViewValueChanged callback

diff --git a/Engine/Plugins/Experimental/ChaosEditor/Source/FractureEditor/Private/FractureEditorModeToolkit.cpp b/Engine/Plugins/Experimental/ChaosEditor/Source/FractureEditor/Private/FractureEditorModeToolkit.cpp
--- a/Engine/Plugins/Experimental/ChaosEditor/Source/FractureEditor/Private/FractureEditorModeToolkit.cpp
+++ b/Engine/Plugins/Experimental/ChaosEditor/Source/FractureEditor/Private/FractureEditorModeToolkit.cpp
@@ -1376,6 +1376,7 @@ void FFractureEditorModeToolkit::OnLevelViewValueChanged()
 			EditBoneColor.SetLevelViewMode(FractureLevel);
 			// Clear selection below currently-selected view level and update highlights,
 			// so the selection is compatible with the current 3D view and outliner (e.g., doesn't hide selection of children)
+			EditBoneColor.Sanitize();
 			EditBoneColor.FilterSelectionToLevel();
 			UpdateExplodedVectors(Comp);
 			Comp->MarkRenderStateDirty();
