Make the fracture Set Removal, State, and Convex tools automatically change the Fracture Mode's Outliner Column Mode to Removal, State and Collision respectively on switching to the tool

Add Toolkit as a parameter to fracture tool setup function (as a weak pointer to match the Execute function, and to allow it to be safely held if needed)

diff --git a/Engine/Plugins/Experimental/ChaosEditor/Source/FractureEditor/Private/FractureEditorModeToolkit.cpp b/Engine/Plugins/Experimental/ChaosEditor/Source/FractureEditor/Private/FractureEditorModeToolkit.cpp
--- a/Engine/Plugins/Experimental/ChaosEditor/Source/FractureEditor/Private/FractureEditorModeToolkit.cpp
+++ b/Engine/Plugins/Experimental/ChaosEditor/Source/FractureEditor/Private/FractureEditorModeToolkit.cpp
@@ -877,6 +877,21 @@ TSharedRef<SDockTab> FFractureEditorModeToolkit::CreateStatisticsTab(const FSpaw
 	return CreatedTab.ToSharedRef();
 }
 
+void FFractureEditorModeToolkit::SetOutlinerColumnMode(EOutlinerColumnMode ColumnMode)
+{
+	UOutlinerSettings* OutlinerSettings = GetMutableDefault<UOutlinerSettings>();
+	OutlinerSettings->ColumnMode = ColumnMode;
+	UpdateOutlinerHeader();
+}
+
+void FFractureEditorModeToolkit::UpdateOutlinerHeader()
+{
+	OutlinerView->RegenerateHeader();
+	FGeometryCollectionStatistics Stats;
+	GetStatisticsSummary(Stats);
+	StatisticsView->SetStatistics(Stats);
+}
+
 void FFractureEditorModeToolkit::OnObjectPostEditChange( UObject* Object, FPropertyChangedEvent& PropertyChangedEvent )
 {
 	if (PropertyChangedEvent.Property)
@@ -911,17 +926,11 @@ void FFractureEditorModeToolkit::OnObjectPostEditChange( UObject* Object, FPrope
 		}
 		else if (PropertyChangedEvent.Property->GetFName() == GET_MEMBER_NAME_CHECKED(UOutlinerSettings, ColorByLevel))
 		{
-			OutlinerView->RegenerateItems();
-			FGeometryCollectionStatistics Stats;
-			GetStatisticsSummary(Stats);
-			StatisticsView->SetStatistics(Stats);
+			UpdateOutlinerHeader();
 		}
 		else if (PropertyChangedEvent.Property->GetFName() == GET_MEMBER_NAME_CHECKED(UOutlinerSettings, ColumnMode))
 		{
-			OutlinerView->RegenerateHeader();
-			FGeometryCollectionStatistics Stats;
-			GetStatisticsSummary(Stats);
-			StatisticsView->SetStatistics(Stats);
+			UpdateOutlinerHeader();
 		}
 	}
 }
@@ -1582,7 +1591,7 @@ void FFractureEditorModeToolkit::SetActiveTool(UFractureModalTool* InActiveTool)
 	{
 		ActiveTool->OnPropertyModifiedDirectlyByTool.AddSP(this, &FFractureEditorModeToolkit::InvalidateCachedDetailPanelState);
 
-		ActiveTool->Setup();
+		ActiveTool->Setup(StaticCastSharedRef<FFractureEditorModeToolkit>(AsShared()));
 
 		Settings.Append(ActiveTool->GetSettingsObjects());
 
diff --git a/Engine/Plugins/Experimental/ChaosEditor/Source/FractureEditor/Private/FractureToolConvex.cpp b/Engine/Plugins/Experimental/ChaosEditor/Source/FractureEditor/Private/FractureToolConvex.cpp
--- a/Engine/Plugins/Experimental/ChaosEditor/Source/FractureEditor/Private/FractureToolConvex.cpp
+++ b/Engine/Plugins/Experimental/ChaosEditor/Source/FractureEditor/Private/FractureToolConvex.cpp
@@ -2,6 +2,7 @@
 
 #include "FractureToolConvex.h"
 
+#include "SGeometryCollectionOutliner.h"
 #include "FractureToolContext.h"
 #include "FractureModeSettings.h"
 
@@ -372,9 +373,13 @@ int32 UFractureToolConvex::ExecuteFracture(const FFractureToolContext& FractureC
 	return INDEX_NONE;
 }
 
-void UFractureToolConvex::Setup()
+void UFractureToolConvex::Setup(TWeakPtr<FFractureEditorModeToolkit> InToolkit)
 {
-	Super::Setup();
+	Super::Setup(InToolkit);
+	if (InToolkit.IsValid())
+	{
+		InToolkit.Pin()->SetOutlinerColumnMode(EOutlinerColumnMode::Collision);
+	}
 }
 
 #undef LOCTEXT_NAMESPACE
diff --git a/Engine/Plugins/Experimental/ChaosEditor/Source/FractureEditor/Private/FractureToolConvex.h b/Engine/Plugins/Experimental/ChaosEditor/Source/FractureEditor/Private/FractureToolConvex.h
--- a/Engine/Plugins/Experimental/ChaosEditor/Source/FractureEditor/Private/FractureToolConvex.h
+++ b/Engine/Plugins/Experimental/ChaosEditor/Source/FractureEditor/Private/FractureToolConvex.h
@@ -129,7 +129,7 @@ public:
 	void PromoteChildren();
 	void ClearCustomConvex();
 
-	virtual void Setup() override;
+	virtual void Setup(TWeakPtr<FFractureEditorModeToolkit> InToolkit) override;
 
 
 	UPROPERTY(EditAnywhere, Category = Convex)
diff --git a/Engine/Plugins/Experimental/ChaosEditor/Source/FractureEditor/Private/FractureToolCustomVoronoi.cpp b/Engine/Plugins/Experimental/ChaosEditor/Source/FractureEditor/Private/FractureToolCustomVoronoi.cpp
--- a/Engine/Plugins/Experimental/ChaosEditor/Source/FractureEditor/Private/FractureToolCustomVoronoi.cpp
+++ b/Engine/Plugins/Experimental/ChaosEditor/Source/FractureEditor/Private/FractureToolCustomVoronoi.cpp
@@ -64,9 +64,9 @@ UFractureToolCustomVoronoi::UFractureToolCustomVoronoi(const FObjectInitializer&
 	GizmoSettings->OwnerTool = this;
 }
 
-void UFractureToolCustomVoronoi::Setup()
+void UFractureToolCustomVoronoi::Setup(TWeakPtr<FFractureEditorModeToolkit> InToolkit)
 {
-	Super::Setup();
+	Super::Setup(InToolkit);
 	GizmoSettings->Setup(this, ETransformGizmoSubElements::FullTranslateRotateScale);
 	// Stop scaling at 0 rather than going negative
 	GizmoSettings->TransformGizmo->SetDisallowNegativeScaling(true);
diff --git a/Engine/Plugins/Experimental/ChaosEditor/Source/FractureEditor/Private/FractureToolCustomVoronoi.h b/Engine/Plugins/Experimental/ChaosEditor/Source/FractureEditor/Private/FractureToolCustomVoronoi.h
--- a/Engine/Plugins/Experimental/ChaosEditor/Source/FractureEditor/Private/FractureToolCustomVoronoi.h
+++ b/Engine/Plugins/Experimental/ChaosEditor/Source/FractureEditor/Private/FractureToolCustomVoronoi.h
@@ -139,7 +139,7 @@ public:
 	UPROPERTY(EditAnywhere, Category = Uniform)
 	TObjectPtr<UFractureTransformGizmoSettings> GizmoSettings;
 
-	virtual void Setup() override;
+	virtual void Setup(TWeakPtr<FFractureEditorModeToolkit> InToolkit) override;
 	virtual void Shutdown() override;
 
 	// Clear "live" voronoi sites (sites that are being edited by current settings / gizmo) -- they will repopulate on next call to GenerateVoronoiSites()
diff --git a/Engine/Plugins/Experimental/ChaosEditor/Source/FractureEditor/Private/FractureToolCutter.h b/Engine/Plugins/Experimental/ChaosEditor/Source/FractureEditor/Private/FractureToolCutter.h
--- a/Engine/Plugins/Experimental/ChaosEditor/Source/FractureEditor/Private/FractureToolCutter.h
+++ b/Engine/Plugins/Experimental/ChaosEditor/Source/FractureEditor/Private/FractureToolCutter.h
@@ -226,9 +226,9 @@ public:
 		UpdateDefaultRandomSeed();
 	}
 
-	virtual void Setup() override
+	virtual void Setup(TWeakPtr<FFractureEditorModeToolkit> InToolkit) override
 	{
-		Super::Setup();
+		Super::Setup(InToolkit);
 		CutterSettings->UpdateActiveMaterialNames(GetSelectedComponentMaterialNames(true));
 	}
 
diff --git a/Engine/Plugins/Experimental/ChaosEditor/Source/FractureEditor/Private/FractureToolMaterials.h b/Engine/Plugins/Experimental/ChaosEditor/Source/FractureEditor/Private/FractureToolMaterials.h
--- a/Engine/Plugins/Experimental/ChaosEditor/Source/FractureEditor/Private/FractureToolMaterials.h
+++ b/Engine/Plugins/Experimental/ChaosEditor/Source/FractureEditor/Private/FractureToolMaterials.h
@@ -128,9 +128,9 @@ public:
 	virtual void SelectedBonesChanged() override;
 
 	// Called when the modal tool is entered
-	virtual void Setup() override
+	virtual void Setup(TWeakPtr<FFractureEditorModeToolkit> InToolkit) override
 	{
-		Super::Setup();
+		Super::Setup(InToolkit);
 		UpdateActiveMaterialsInfo();
 	}
 
diff --git a/Engine/Plugins/Experimental/ChaosEditor/Source/FractureEditor/Private/FractureToolPlaneCut.cpp b/Engine/Plugins/Experimental/ChaosEditor/Source/FractureEditor/Private/FractureToolPlaneCut.cpp
--- a/Engine/Plugins/Experimental/ChaosEditor/Source/FractureEditor/Private/FractureToolPlaneCut.cpp
+++ b/Engine/Plugins/Experimental/ChaosEditor/Source/FractureEditor/Private/FractureToolPlaneCut.cpp
@@ -140,9 +140,9 @@ UFractureToolPlaneCut::UFractureToolPlaneCut(const FObjectInitializer& ObjInit)
 }
 
 
-void UFractureToolPlaneCut::Setup()
+void UFractureToolPlaneCut::Setup(TWeakPtr<FFractureEditorModeToolkit> InToolkit)
 {
-	Super::Setup();
+	Super::Setup(InToolkit);
 	GizmoSettings->Setup(this);
 	PlaneCutSettings->bCanCutWithMultiplePlanes = !GizmoSettings->bUseGizmo;
 	NotifyOfPropertyChangeByTool(PlaneCutSettings);
diff --git a/Engine/Plugins/Experimental/ChaosEditor/Source/FractureEditor/Private/FractureToolPlaneCut.h b/Engine/Plugins/Experimental/ChaosEditor/Source/FractureEditor/Private/FractureToolPlaneCut.h
--- a/Engine/Plugins/Experimental/ChaosEditor/Source/FractureEditor/Private/FractureToolPlaneCut.h
+++ b/Engine/Plugins/Experimental/ChaosEditor/Source/FractureEditor/Private/FractureToolPlaneCut.h
@@ -57,7 +57,7 @@ public:
 	virtual void FractureContextChanged() override;
 	virtual int32 ExecuteFracture(const FFractureToolContext& FractureContext) override;
 
-	virtual void Setup() override;
+	virtual void Setup(TWeakPtr<FFractureEditorModeToolkit> InToolkit) override;
 	virtual void Shutdown() override;
 
 	// IDynamicMeshOperatorFactory API, for generating noise preview meshes
diff --git a/Engine/Plugins/Experimental/ChaosEditor/Source/FractureEditor/Private/FractureToolProperties.cpp b/Engine/Plugins/Experimental/ChaosEditor/Source/FractureEditor/Private/FractureToolProperties.cpp
--- a/Engine/Plugins/Experimental/ChaosEditor/Source/FractureEditor/Private/FractureToolProperties.cpp
+++ b/Engine/Plugins/Experimental/ChaosEditor/Source/FractureEditor/Private/FractureToolProperties.cpp
@@ -2,6 +2,8 @@
 
 #include "FractureToolProperties.h"
 
+#include "SGeometryCollectionOutliner.h"
+
 #include "GeometryCollection/GeometryCollection.h"
 #include "GeometryCollection/GeometryCollectionComponent.h"
 #include "GeometryCollection/GeometryCollectionEngineRemoval.h"
@@ -58,6 +60,15 @@ void UFractureToolSetInitialDynamicState::RegisterUICommand(FFractureEditorComma
 	BindingContext->SetInitialDynamicState = UICommandInfo;
 }
 
+void UFractureToolSetInitialDynamicState::Setup(TWeakPtr<FFractureEditorModeToolkit> InToolkit)
+{
+	Super::Setup(InToolkit);
+	if (InToolkit.IsValid())
+	{
+		InToolkit.Pin()->SetOutlinerColumnMode(EOutlinerColumnMode::State);
+	}
+}
+
 void UFractureToolSetInitialDynamicState::Execute(TWeakPtr<FFractureEditorModeToolkit> InToolkit)
 {
 	if (InToolkit.IsValid())
@@ -131,6 +142,15 @@ void UFractureToolSetRemoveOnBreak::RegisterUICommand(FFractureEditorCommands* B
 	BindingContext->SetRemoveOnBreak = UICommandInfo;
 }
 
+void UFractureToolSetRemoveOnBreak::Setup(TWeakPtr<FFractureEditorModeToolkit> InToolkit)
+{
+	Super::Setup(InToolkit);
+	if (InToolkit.IsValid())
+	{
+		InToolkit.Pin()->SetOutlinerColumnMode(EOutlinerColumnMode::Removal);
+	}
+}
+
 void UFractureToolSetRemoveOnBreak::Execute(TWeakPtr<FFractureEditorModeToolkit> InToolkit)
 {
 	Toolkit = InToolkit;
diff --git a/Engine/Plugins/Experimental/ChaosEditor/Source/FractureEditor/Private/FractureToolProperties.h b/Engine/Plugins/Experimental/ChaosEditor/Source/FractureEditor/Private/FractureToolProperties.h
--- a/Engine/Plugins/Experimental/ChaosEditor/Source/FractureEditor/Private/FractureToolProperties.h
+++ b/Engine/Plugins/Experimental/ChaosEditor/Source/FractureEditor/Private/FractureToolProperties.h
@@ -54,6 +54,7 @@ public:
 
 	virtual void RegisterUICommand(FFractureEditorCommands* BindingContext) override;
 
+	virtual void Setup(TWeakPtr<FFractureEditorModeToolkit> InToolkit) override;
 	virtual void Execute(TWeakPtr<FFractureEditorModeToolkit> InToolkit) override;
 
 	UPROPERTY(EditAnywhere, Category = InitialDynamicState)
@@ -117,6 +118,8 @@ public:
 
 	virtual void RegisterUICommand(FFractureEditorCommands* BindingContext) override;
 
+	virtual void Setup(TWeakPtr<FFractureEditorModeToolkit> InToolkit) override;
+
 	virtual void Execute(TWeakPtr<FFractureEditorModeToolkit> InToolkit) override;
 
 	void DeleteRemoveOnBreakData();
diff --git a/Engine/Plugins/Experimental/ChaosEditor/Source/FractureEditor/Private/FractureToolProximity.cpp b/Engine/Plugins/Experimental/ChaosEditor/Source/FractureEditor/Private/FractureToolProximity.cpp
--- a/Engine/Plugins/Experimental/ChaosEditor/Source/FractureEditor/Private/FractureToolProximity.cpp
+++ b/Engine/Plugins/Experimental/ChaosEditor/Source/FractureEditor/Private/FractureToolProximity.cpp
@@ -242,9 +242,9 @@ int32 UFractureToolProximity::ExecuteFracture(const FFractureToolContext& Fractu
 	return INDEX_NONE;
 }
 
-void UFractureToolProximity::Setup()
+void UFractureToolProximity::Setup(TWeakPtr<FFractureEditorModeToolkit> InToolkit)
 {
-	Super::Setup();
+	Super::Setup(InToolkit);
 }
 
 #undef LOCTEXT_NAMESPACE
diff --git a/Engine/Plugins/Experimental/ChaosEditor/Source/FractureEditor/Private/FractureToolProximity.h b/Engine/Plugins/Experimental/ChaosEditor/Source/FractureEditor/Private/FractureToolProximity.h
--- a/Engine/Plugins/Experimental/ChaosEditor/Source/FractureEditor/Private/FractureToolProximity.h
+++ b/Engine/Plugins/Experimental/ChaosEditor/Source/FractureEditor/Private/FractureToolProximity.h
@@ -126,7 +126,7 @@ public:
 
 	virtual TArray<FFractureToolContext> GetFractureToolContexts() const override;
 
-	virtual void Setup() override;
+	virtual void Setup(TWeakPtr<FFractureEditorModeToolkit> InToolkit) override;
 
 	UPROPERTY(EditAnywhere, Category = Proximity)
 	TObjectPtr<UFractureProximitySettings> ProximitySettings;
diff --git a/Engine/Plugins/Experimental/ChaosEditor/Source/FractureEditor/Private/FractureToolRadial.cpp b/Engine/Plugins/Experimental/ChaosEditor/Source/FractureEditor/Private/FractureToolRadial.cpp
--- a/Engine/Plugins/Experimental/ChaosEditor/Source/FractureEditor/Private/FractureToolRadial.cpp
+++ b/Engine/Plugins/Experimental/ChaosEditor/Source/FractureEditor/Private/FractureToolRadial.cpp
@@ -19,9 +19,9 @@ UFractureToolRadial::UFractureToolRadial(const FObjectInitializer& ObjInit)
 	GizmoSettings->OwnerTool = this;
 }
 
-void UFractureToolRadial::Setup()
+void UFractureToolRadial::Setup(TWeakPtr<FFractureEditorModeToolkit> InToolkit)
 {
-	Super::Setup();
+	Super::Setup(InToolkit);
 	GizmoSettings->Setup(this);
 }
 
diff --git a/Engine/Plugins/Experimental/ChaosEditor/Source/FractureEditor/Private/FractureToolRadial.h b/Engine/Plugins/Experimental/ChaosEditor/Source/FractureEditor/Private/FractureToolRadial.h
--- a/Engine/Plugins/Experimental/ChaosEditor/Source/FractureEditor/Private/FractureToolRadial.h
+++ b/Engine/Plugins/Experimental/ChaosEditor/Source/FractureEditor/Private/FractureToolRadial.h
@@ -103,7 +103,7 @@ public:
 	UPROPERTY(EditAnywhere, Category = Uniform)
 	TObjectPtr<UFractureTransformGizmoSettings> GizmoSettings;
 
-	virtual void Setup() override;
+	virtual void Setup(TWeakPtr<FFractureEditorModeToolkit> InToolkit) override;
 	virtual void Shutdown() override;
 
 	virtual void UpdateUseGizmo(bool bUseGizmo) override
diff --git a/Engine/Plugins/Experimental/ChaosEditor/Source/FractureEditor/Public/FractureEditorModeToolkit.h b/Engine/Plugins/Experimental/ChaosEditor/Source/FractureEditor/Public/FractureEditorModeToolkit.h
--- a/Engine/Plugins/Experimental/ChaosEditor/Source/FractureEditor/Public/FractureEditorModeToolkit.h
+++ b/Engine/Plugins/Experimental/ChaosEditor/Source/FractureEditor/Public/FractureEditorModeToolkit.h
@@ -10,6 +10,7 @@
 #include "Widgets/Input/SCheckBox.h"
 #include "Widgets/Layout/SSplitter.h"
 #include "IDetailCustomization.h"
+#include "SGeometryCollectionOutliner.h"
 
 class IDetailsView;
 class IPropertyHandle;
@@ -188,6 +189,8 @@ public:
 	TSharedPtr<SWidget> LevelViewWidget;
 	TSharedPtr<SWidget> ShowBoneColorsWidget;
 
+	void SetOutlinerColumnMode(EOutlinerColumnMode ColumnMode);
+
 protected:
 	/** FModeToolkit interface */
 	virtual void RequestModeUITabs() override;
@@ -221,6 +224,8 @@ private:
 	void UpdateAssetLocationMode(TSharedPtr<FString> NewString);
 	void UpdateAssetPanelFromSettings();
 	void OnProjectSettingsModified();
+
+	void UpdateOutlinerHeader();
 	
 private:
 	UFractureModalTool* ActiveTool;
diff --git a/Engine/Plugins/Experimental/ChaosEditor/Source/FractureEditor/Public/FractureTool.h b/Engine/Plugins/Experimental/ChaosEditor/Source/FractureEditor/Public/FractureTool.h
--- a/Engine/Plugins/Experimental/ChaosEditor/Source/FractureEditor/Public/FractureTool.h
+++ b/Engine/Plugins/Experimental/ChaosEditor/Source/FractureEditor/Public/FractureTool.h
@@ -186,7 +186,7 @@ public:
 	virtual void SelectedBonesChanged() {}
 
 	// Called when the modal tool is entered
-	virtual void Setup()
+	virtual void Setup(TWeakPtr<FFractureEditorModeToolkit> InToolkit)
 	{
 		GEngine->OnComponentTransformChanged().AddUObject(this, &UFractureModalTool::OnComponentTransformChangedInternal);
 	}
