CIS fix -- pull EColumnMode enum out to public FractureEditorModeToolkit.h so it can be accessed as needed

diff --git a/Engine/Plugins/Experimental/ChaosEditor/Source/FractureEditor/Private/SGeometryCollectionOutliner.h b/Engine/Plugins/Experimental/ChaosEditor/Source/FractureEditor/Private/SGeometryCollectionOutliner.h
--- a/Engine/Plugins/Experimental/ChaosEditor/Source/FractureEditor/Private/SGeometryCollectionOutliner.h
+++ b/Engine/Plugins/Experimental/ChaosEditor/Source/FractureEditor/Private/SGeometryCollectionOutliner.h
@@ -18,6 +18,8 @@
 #include "GeometryCollection/GeometryCollectionComponent.h"
 #include "GeometryCollection/Facades/CollectionRemoveOnBreakFacade.h"
 
+#include "FractureEditorModeToolkit.h"
+
 #include "SGeometryCollectionOutliner.generated.h"
 
 class FGeometryCollection;
@@ -37,16 +39,6 @@ enum class EOutlinerItemNameEnum : uint8
 	BoneIndex = 1					UMETA(DisplayName = "Bone Index"),
 };
 
-UENUM(BlueprintType)
-enum class EOutlinerColumnMode : uint8
-{
-	State = 0				UMETA(DisplayName = "State"),
-	Damage = 1				UMETA(DisplayName = "Damage"),
-	Removal = 2				UMETA(DisplayName = "Removal"),
-	Collision = 3			UMETA(DisplayName = "Collision"),
-	Size = 4				UMETA(DisplayName = "Size"),
-};
-
 /** Settings for Outliner configuration. **/
 UCLASS()
 class UOutlinerSettings : public UObject
diff --git a/Engine/Plugins/Experimental/ChaosEditor/Source/FractureEditor/Public/FractureEditorModeToolkit.h b/Engine/Plugins/Experimental/ChaosEditor/Source/FractureEditor/Public/FractureEditorModeToolkit.h
--- a/Engine/Plugins/Experimental/ChaosEditor/Source/FractureEditor/Public/FractureEditorModeToolkit.h
+++ b/Engine/Plugins/Experimental/ChaosEditor/Source/FractureEditor/Public/FractureEditorModeToolkit.h
@@ -10,7 +10,8 @@
 #include "Widgets/Input/SCheckBox.h"
 #include "Widgets/Layout/SSplitter.h"
 #include "IDetailCustomization.h"
-#include "SGeometryCollectionOutliner.h"
+
+#include "FractureEditorModeToolkit.generated.h"
 
 class IDetailsView;
 class IPropertyHandle;
@@ -86,6 +87,16 @@ struct FTextAndSlateColor
 	FSlateColor Color;
 };
 
+UENUM(BlueprintType)
+enum class EOutlinerColumnMode : uint8
+{
+	State = 0				UMETA(DisplayName = "State"),
+	Damage = 1				UMETA(DisplayName = "Damage"),
+	Removal = 2				UMETA(DisplayName = "Removal"),
+	Collision = 3			UMETA(DisplayName = "Collision"),
+	Size = 4				UMETA(DisplayName = "Size"),
+};
+
 class FRACTUREEDITOR_API FFractureEditorModeToolkit : public FModeToolkit, public FGCObject
 {
 public:
