Make fracture mode's interactive selection tool set the outliner column mode to 'size'

diff --git a/Engine/Plugins/Experimental/ChaosEditor/Source/FractureEditor/Private/FractureToolSelection.cpp b/Engine/Plugins/Experimental/ChaosEditor/Source/FractureEditor/Private/FractureToolSelection.cpp
--- a/Engine/Plugins/Experimental/ChaosEditor/Source/FractureEditor/Private/FractureToolSelection.cpp
+++ b/Engine/Plugins/Experimental/ChaosEditor/Source/FractureEditor/Private/FractureToolSelection.cpp
@@ -322,6 +322,15 @@ void URectangleMarqueeManager::DrawHUD(FCanvas* Canvas, bool bThisViewHasFocus)
 }
 
 
+void UFractureToolSelection::Setup(TWeakPtr<FFractureEditorModeToolkit> InToolkit)
+{
+	Super::Setup(InToolkit);
+	if (InToolkit.IsValid())
+	{
+		InToolkit.Pin()->SetOutlinerColumnMode(EOutlinerColumnMode::Size);
+	}
+}
+
 void UFractureToolSelection::Shutdown()
 {
 	DestroyRectangleSelectionBehavior();
diff --git a/Engine/Plugins/Experimental/ChaosEditor/Source/FractureEditor/Private/FractureToolSelection.h b/Engine/Plugins/Experimental/ChaosEditor/Source/FractureEditor/Private/FractureToolSelection.h
--- a/Engine/Plugins/Experimental/ChaosEditor/Source/FractureEditor/Private/FractureToolSelection.h
+++ b/Engine/Plugins/Experimental/ChaosEditor/Source/FractureEditor/Private/FractureToolSelection.h
@@ -225,6 +225,7 @@ public:
 	virtual void Execute(TWeakPtr<FFractureEditorModeToolkit> InToolkit) override;
 	void UpdateSelection(FFractureToolContext& FractureContext);
 
+	virtual void Setup(TWeakPtr<FFractureEditorModeToolkit> InToolkit) override;
 	virtual void Shutdown() override;
 
 	// IClickBehaviorTarget implementation
