Make Fracture Mode's Convex Tool only show collision for the current Fracture Level

diff --git a/Engine/Plugins/Experimental/ChaosEditor/Source/FractureEditor/Private/FractureToolContext.cpp b/Engine/Plugins/Experimental/ChaosEditor/Source/FractureEditor/Private/FractureToolContext.cpp
--- a/Engine/Plugins/Experimental/ChaosEditor/Source/FractureEditor/Private/FractureToolContext.cpp
+++ b/Engine/Plugins/Experimental/ChaosEditor/Source/FractureEditor/Private/FractureToolContext.cpp
@@ -225,6 +225,52 @@ void FFractureToolContext::ConvertSelectionToEmbeddedGeometryNodes(int32 Index,
 	}
 }
 
+void FFractureToolContext::ConvertSelectionToLevel(int32 TargetLevel, bool bAllowSelectChildren, bool bAllowSelectParents)
+{
+	const TManagedArray<TSet<int32>>& Children = GeometryCollection->Children;
+	const TManagedArray<int32>& Parents = GeometryCollection->Parent;
+	const TManagedArray<int32>& SimulationType = GeometryCollection->SimulationType;
+	const TManagedArray<int32>* Levels = GeometryCollection->FindAttribute<int32>("Level", FGeometryCollection::TransformGroup);
+
+	if (TargetLevel < 0 || !Levels)
+	{
+		return; // No levels to filter by
+	}
+
+	TSet<int32> FilteredSelection;
+	const int32 NumBones = GeometryCollection->Transform.Num();
+	// Note SafetyMaxLoops is just to prevent infinite loop in case of malformed data (e.g., a cycle in the Children/Parent indices)
+	int32 SafetyMaxLoops = 2 * SelectedBones.Num() * NumBones;
+	while (SelectedBones.Num() && SafetyMaxLoops-- > 0)
+	{
+		int32 BoneIdx = SelectedBones.Pop();
+		if (BoneIdx < 0 || BoneIdx >= NumBones)
+		{
+			continue;
+		}
+		int32 BoneLevel = (*Levels)[BoneIdx];
+		if (BoneLevel == TargetLevel)
+		{
+			FilteredSelection.Add(BoneIdx);
+		}
+		// Optionally allow the selection to expand clusters to reach the target level
+		else if (bAllowSelectChildren && TargetLevel > BoneLevel && SimulationType[BoneIdx] == FGeometryCollection::ESimulationTypes::FST_Clustered)
+		{
+			SelectedBones.Append(GeometryCollection->Children[BoneIdx].Array());
+		}
+		// Optionally allow the selection to travel to parents to reach the target level
+		else if (bAllowSelectParents && TargetLevel < BoneLevel)
+		{
+			int32 Parent = GeometryCollection->Parent[BoneIdx];
+			if (Parent > -1)
+			{
+				SelectedBones.Add(Parent);
+			}
+		}
+	}
+	SelectedBones = FilteredSelection.Array();
+}
+
 void FFractureToolContext::ConvertSelectionToClusterNodes()
 {
 	// If this is a non-cluster node, select the cluster containing it instead.
diff --git a/Engine/Plugins/Experimental/ChaosEditor/Source/FractureEditor/Private/FractureToolContext.h b/Engine/Plugins/Experimental/ChaosEditor/Source/FractureEditor/Private/FractureToolContext.h
--- a/Engine/Plugins/Experimental/ChaosEditor/Source/FractureEditor/Private/FractureToolContext.h
+++ b/Engine/Plugins/Experimental/ChaosEditor/Source/FractureEditor/Private/FractureToolContext.h
@@ -24,6 +24,7 @@ public:
 	void ConvertSelectionToRigidNodes();
 	void ConvertSelectionToEmbeddedGeometryNodes();
 	void ConvertSelectionToClusterNodes();
+	void ConvertSelectionToLevel(int32 TargetLevel, bool bAllowSelectChildren, bool bAllowSelectParents);
 	void RemoveChildrenOfSelectedNodesFromSelection();
 	void ConvertEmbeddedSelectionToParents();
 	void RemoveRootNodes();
diff --git a/Engine/Plugins/Experimental/ChaosEditor/Source/FractureEditor/Private/FractureToolConvex.cpp b/Engine/Plugins/Experimental/ChaosEditor/Source/FractureEditor/Private/FractureToolConvex.cpp
--- a/Engine/Plugins/Experimental/ChaosEditor/Source/FractureEditor/Private/FractureToolConvex.cpp
+++ b/Engine/Plugins/Experimental/ChaosEditor/Source/FractureEditor/Private/FractureToolConvex.cpp
@@ -5,6 +5,7 @@
 #include "SGeometryCollectionOutliner.h"
 #include "FractureToolContext.h"
 #include "FractureModeSettings.h"
+#include "FractureSettings.h"
 
 #include "GeometryCollection/GeometryCollectionObject.h"
 #include "GeometryCollection/GeometryCollectionComponent.h"
@@ -296,7 +297,9 @@ void UFractureToolConvex::FractureContextChanged()
 
 	ClearVisualizations();
 
-	for (const FFractureToolContext& FractureContext : FractureContexts)
+	const UFractureSettings* FractureSettings = GetDefault<UFractureSettings>();
+
+	for (FFractureToolContext& FractureContext : FractureContexts)
 	{
 		if (UGeometryCollection* RestCollection = FractureContext.GetFracturedGeometryCollection())
 		{
@@ -307,6 +310,9 @@ void UFractureToolConvex::FractureContextChanged()
 			//FTransform OuterTransform = FractureContext.GetTransform();
 			const FTransform OuterTransform = FractureContext.GetGeometryCollectionComponent()->GetComponentTransform();
 
+			FractureContext.ConvertSelectionToLevel(FractureSettings->FractureLevel, 
+				true /*Auto-select children at the target level*/, false /*Don't auto-select parents*/);
+
 			for (int32 TransformIdx : FractureContext.GetSelection())
 			{
 				EdgesMappings.AddMapping(CollectionIdx, TransformIdx, HullEdges.Num());
