Fix potential crash when computing the exploding vectors in the fracture editor toolkit

diff --git a/Engine/Plugins/Experimental/ChaosEditor/Source/FractureEditor/Private/FractureEditorModeToolkit.cpp b/Engine/Plugins/Experimental/ChaosEditor/Source/FractureEditor/Private/FractureEditorModeToolkit.cpp
--- a/Engine/Plugins/Experimental/ChaosEditor/Source/FractureEditor/Private/FractureEditorModeToolkit.cpp
+++ b/Engine/Plugins/Experimental/ChaosEditor/Source/FractureEditor/Private/FractureEditorModeToolkit.cpp
@@ -1905,6 +1905,11 @@ void FFractureEditorModeToolkit::UpdateExplodedVectors(UGeometryCollectionCompon
 	}
 #endif
 
+	if (GeometryCollectionComponent == nullptr || GeometryCollectionComponent->GetRestCollection() == nullptr)
+	{
+		return;
+	}
+
 	TSharedPtr<FGeometryCollection, ESPMode::ThreadSafe> GeometryCollectionPtr = GeometryCollectionComponent->GetRestCollection()->GetGeometryCollection();
 	const FGeometryCollection* OutGeometryCollectionConst = GeometryCollectionPtr.Get();
 
