add optional grid visualization for grid-based autocluster

diff --git a/Engine/Plugins/Experimental/ChaosEditor/Source/FractureEditor/Private/FractureToolAutoCluster.cpp b/Engine/Plugins/Experimental/ChaosEditor/Source/FractureEditor/Private/FractureToolAutoCluster.cpp
--- a/Engine/Plugins/Experimental/ChaosEditor/Source/FractureEditor/Private/FractureToolAutoCluster.cpp
+++ b/Engine/Plugins/Experimental/ChaosEditor/Source/FractureEditor/Private/FractureToolAutoCluster.cpp
@@ -97,6 +97,15 @@ void UFractureToolAutoCluster::Render(const FSceneView* View, FViewport* Viewpor
 				FVector Pt = ShowGridPoints[Idx];
 				PDI->DrawPoint(Pt, FLinearColor(1.f, 0.f, 1.f), 5, SDPG_Foreground);
 			});
+		if (AutoClusterSettings->bShowGridLines)
+		{
+			EnumerateVisualizationMapping(GridLinesMappings, ShowGridLines.Num(),
+				[&](int32 Idx, FVector ExplodedVector)
+				{
+					TPair<FVector, FVector> Pts = ShowGridLines[Idx];
+					PDI->DrawLine(Pts.Key, Pts.Value, FLinearColor(1.f, 1.f, 0.f), SDPG_Foreground);
+				});
+		}
 	}
 
 	if (AutoClusterSettings->ClusterSizeMethod == EClusterSizeMethod::BySize)
@@ -128,15 +137,44 @@ void UFractureToolAutoCluster::FractureContextChanged()
 			for (const int32 ClusterIndex : Context.GetSelection())
 			{
 				GridPointsMappings.AddMapping(CollectionIdx, ClusterIndex, ShowGridPoints.Num());
+				FBox Bounds;
 				TArray<FVector> LocalGridPoints = FFractureEngineClustering::GenerateGridSites(
 					*GeometryCollection, ClusterIndex,
 					AutoClusterSettings->ClusterGridWidth,
 					AutoClusterSettings->ClusterGridDepth,
-					AutoClusterSettings->ClusterGridHeight);
+					AutoClusterSettings->ClusterGridHeight,
+					&Bounds);
 				for (FVector LocalPt : LocalGridPoints)
 				{
 					ShowGridPoints.Add(ContextTransform.TransformPosition(LocalPt));
 				}
+				GridLinesMappings.AddMapping(CollectionIdx, ClusterIndex, ShowGridLines.Num());
+				FIntVector3 Dims(
+					AutoClusterSettings->ClusterGridWidth,
+					AutoClusterSettings->ClusterGridDepth,
+					AutoClusterSettings->ClusterGridHeight);
+				FVector InvDims(1.0 / double(Dims.X), 1.0 / double(Dims.Y), 1.0 / double(Dims.Z));
+				FVector Diag = Bounds.Max - Bounds.Min;
+				for (int32 Dim = 0; Dim < 3; ++Dim)
+				{
+					FIntVector3 DMap(Dim, (Dim + 1) % 3, (Dim + 2) % 3);
+					FVector Start, End;
+					Start[DMap.Z] = Bounds.Min[DMap.Z];
+					End[DMap.Z] = Bounds.Max[DMap.Z];
+					for (int32 W = 0; W <= Dims[DMap.X]; ++W)
+					{
+						double AlongW = W * InvDims[DMap.X];
+						Start[DMap.X] = Bounds.Min[DMap.X] + AlongW * Diag[DMap.X];
+						End[DMap.X] = Start[DMap.X];
+						for (int32 H = 0; H <= Dims[DMap.Y]; ++H)
+						{
+							double AlongH = H * InvDims[DMap.Y];
+							Start[DMap.Y] = Bounds.Min[DMap.Y] + AlongH * Diag[DMap.Y];
+							End[DMap.Y] = Start[DMap.Y];
+							ShowGridLines.Emplace(ContextTransform.TransformPosition(Start), ContextTransform.TransformPosition(End));
+						}
+					}
+				}
 			}
 		}
 	}
diff --git a/Engine/Plugins/Experimental/ChaosEditor/Source/FractureEditor/Private/FractureToolAutoCluster.h b/Engine/Plugins/Experimental/ChaosEditor/Source/FractureEditor/Private/FractureToolAutoCluster.h
--- a/Engine/Plugins/Experimental/ChaosEditor/Source/FractureEditor/Private/FractureToolAutoCluster.h
+++ b/Engine/Plugins/Experimental/ChaosEditor/Source/FractureEditor/Private/FractureToolAutoCluster.h
@@ -78,6 +78,10 @@ public:
 	UPROPERTY(EditAnywhere, Category = ClusterSize, meta = (ClampMin = "1", UIMax = "20", EditCondition = "ClusterSizeMethod == EClusterSizeMethod::ByGrid", EditConditionHides))
 	int ClusterGridHeight = 2;
 
+	/** If true, show the cluster grid boundary lines. */
+	UPROPERTY(EditAnywhere, Category = ClusterSize, meta = (EditCondition = "ClusterSizeMethod == EClusterSizeMethod::ByGrid", EditConditionHides))
+	bool bShowGridLines = true;
+
 	/** If a cluster has volume less than this value (in cm) cubed, then the auto-cluster process will attempt to merge it into a neighboring cluster. */
 	UPROPERTY(EditAnywhere, Category = ClusterSize, meta = (ClampMin = "0"))
 	float MinimumSize = 0;
@@ -131,11 +135,14 @@ protected:
 		Super::ClearVisualizations();
 		ShowGridPoints.Empty();
 		GridPointsMappings.Empty();
+		ShowGridLines.Empty();
+		GridLinesMappings.Empty();
 	}
 
 private:
 
 	TArray<FVector> ShowGridPoints;
-	FVisualizationMappings GridPointsMappings;
+	TArray<TPair<FVector, FVector>> ShowGridLines;
+	FVisualizationMappings GridPointsMappings, GridLinesMappings;
 };
 
diff --git a/Engine/Plugins/Experimental/Fracture/Source/FractureEngine/Private/FractureEngineClustering.cpp b/Engine/Plugins/Experimental/Fracture/Source/FractureEngine/Private/FractureEngineClustering.cpp
--- a/Engine/Plugins/Experimental/Fracture/Source/FractureEngine/Private/FractureEngineClustering.cpp
+++ b/Engine/Plugins/Experimental/Fracture/Source/FractureEngine/Private/FractureEngineClustering.cpp
@@ -706,11 +706,16 @@ TArray<FVector> FFractureEngineClustering::GenerateGridSites(
 	const int32 ClusterIndex,
 	const int32 InGridX,
 	const int32 InGridY,
-	const int32 InGridZ)
+	const int32 InGridZ,
+	FBox* OutBounds)
 {
 	TArray<FVector> PartitionPositions;
 
 	FBox Bounds = FVoronoiPartitioner::GenerateBounds(&GeometryCollection, ClusterIndex);
+	if (OutBounds)
+	{
+		*OutBounds = Bounds;
+	}
 	int32 GridX = FMath::Max(1, InGridX);
 	int32 GridY = FMath::Max(1, InGridY);
 	int32 GridZ = FMath::Max(1, InGridZ);
diff --git a/Engine/Plugins/Experimental/Fracture/Source/FractureEngine/Public/FractureEngineClustering.h b/Engine/Plugins/Experimental/Fracture/Source/FractureEngine/Public/FractureEngineClustering.h
--- a/Engine/Plugins/Experimental/Fracture/Source/FractureEngine/Public/FractureEngineClustering.h
+++ b/Engine/Plugins/Experimental/Fracture/Source/FractureEngine/Public/FractureEngineClustering.h
@@ -128,7 +128,8 @@ public:
 		const int32 ClusterIndex,
 		const int32 GridX,
 		const int32 GridY,
-		const int32 GridZ);
+		const int32 GridZ,
+		FBox* OutBounds = nullptr);
 };
 
 #if UE_ENABLE_INCLUDE_ORDER_DEPRECATED_IN_5_2
