Fracture Editor : fix potential crash when selection change

diff --git a/Engine/Plugins/Experimental/ChaosEditor/Source/FractureEditor/Private/FractureEditorModeToolkit.cpp b/Engine/Plugins/Experimental/ChaosEditor/Source/FractureEditor/Private/FractureEditorModeToolkit.cpp
--- a/Engine/Plugins/Experimental/ChaosEditor/Source/FractureEditor/Private/FractureEditorModeToolkit.cpp
+++ b/Engine/Plugins/Experimental/ChaosEditor/Source/FractureEditor/Private/FractureEditorModeToolkit.cpp
@@ -1654,16 +1654,17 @@ void FFractureEditorModeToolkit::SetOutlinerComponents(const TArray<UGeometryCol
 
 		if (IsValid(FracturedGeometryCollection)) // Prevents crash when GC is deleted from content browser and actor is selected.
 		{
-			TSharedPtr<FGeometryCollection, ESPMode::ThreadSafe> GeometryCollectionPtr = FracturedGeometryCollection->GetGeometryCollection();
-
-			FGeometryCollectionClusteringUtility::UpdateHierarchyLevelOfChildren(GeometryCollectionPtr.Get(), -1);
-			UpdateExplodedVectors(Component);
-			UpdateHideForComponent(Component);
+			if (TSharedPtr<FGeometryCollection, ESPMode::ThreadSafe> GeometryCollectionPtr = FracturedGeometryCollection->GetGeometryCollection())
+			{
+				FGeometryCollectionClusteringUtility::UpdateHierarchyLevelOfChildren(GeometryCollectionPtr.Get(), -1);
+				UpdateExplodedVectors(Component);
+				UpdateHideForComponent(Component);
 
-			UpdateGeometryComponentAttributes(Component);
-			ComponentsToEdit.Add(Component);
+				UpdateGeometryComponentAttributes(Component);
+				ComponentsToEdit.Add(Component);
 
-			Component->MarkRenderStateDirty();
+				Component->MarkRenderStateDirty();
+			}
 		}	
 	}
 
