Add workaround fix for undo of creating a new geometry collection leaving a broken asset -- just stop trying to transact the geometry collection asset creation

diff --git a/Engine/Plugins/Experimental/ChaosEditor/Source/FractureEditor/Private/FractureToolGenerators.cpp b/Engine/Plugins/Experimental/ChaosEditor/Source/FractureEditor/Private/FractureToolGenerators.cpp
--- a/Engine/Plugins/Experimental/ChaosEditor/Source/FractureEditor/Private/FractureToolGenerators.cpp
+++ b/Engine/Plugins/Experimental/ChaosEditor/Source/FractureEditor/Private/FractureToolGenerators.cpp
@@ -14,7 +14,6 @@
 #include "Engine/StaticMeshActor.h"
 #include "Framework/Application/SlateApplication.h"
 #include "Misc/Change.h"
-#include "ChangeTransactor.h"
 
 #include "FractureModeSettings.h"
 
@@ -442,27 +441,6 @@ bool SCreateGeometryCollectionFromObject::IsCreateAssetFromActorEnabled() const
 	return !bIsReportingError;
 }
 
-// Creates an undo/redo action that (un)registers and object with the Asset Registry.
-// Upon undo this causes the object to be unregistered and as a result be removed from
-// Content Browsers.
-class FAssetRegistrationChange final : public FCommandChange
-{
-public:
-	void Apply(UObject* Object) override
-	{
-		FAssetRegistryModule::AssetCreated(Object);
-	}
-	void Revert(UObject* Object) override
-	{
-		FAssetRegistryModule::AssetDeleted(Object);
-	}
-
-	FString ToString() const override
-	{
-		return TEXT("Asset registry from " LOCTEXT_NAMESPACE);
-	}
-};
-
 FText UFractureToolGenerateAsset::GetDisplayText() const
 {
 	return FText(NSLOCTEXT("Fracture", "FractureToolGenerateAsset", "New"));
@@ -615,8 +593,6 @@ void UFractureToolGenerateAsset::OpenGenerateAssetDialog(TArray<AActor*>& Actors
 
 void UFractureToolGenerateAsset::OnGenerateAssetPathChosen(const FString& InAssetPath, bool bFromToMeshTool, bool bSplitComponents, TArray<AActor*> Actors)
 {	
-	FScopedTransaction Transaction(LOCTEXT("GenerateAsset", "Generate Geometry Collection Asset"));
-
 	//Record the path
 	int32 LastSlash = INDEX_NONE;
 	if (InAssetPath.FindLastChar('/', LastSlash))
@@ -666,6 +642,10 @@ void UFractureToolGenerateAsset::OnGenerateAssetPathChosen(const FString& InAsse
 		GeometryCollectionComponent->MarkRenderDynamicDataDirty();
 		GeometryCollectionComponent->MarkRenderStateDirty();
 
+
+
+		FScopedTransaction Transaction(LOCTEXT("RemoveSourceActors", "Remove Source Actors"));
+
 		for (AActor* Actor : Actors)
 		{
 			Actor->Modify();
@@ -788,12 +768,6 @@ class AGeometryCollectionActor* UFractureToolGenerateAsset::CreateNewGeometryAct
 	UGeometryCollection* InGeometryCollection = static_cast<UGeometryCollection*>(NewObject<UGeometryCollection>(Package, UGeometryCollection::StaticClass(), FName(*UniqueAssetName), RF_Transactional | RF_Public | RF_Standalone));
 	if(!InGeometryCollection->SizeSpecificData.Num()) InGeometryCollection->SizeSpecificData.Add(FGeometryCollectionSizeSpecificData());
 
-	// Record the creation of the geometry collection so it's removed from the Asset Registry and the Content Browser when undo is called.
-	UE::FChangeTransactor Transactor(InGeometryCollection);
-	Transactor.OpenTransaction(LOCTEXT("GeometryCollectionAssetRegistration", "Geometry Collection Asset Registration"));
-	Transactor.AddTransactionChange<FAssetRegistrationChange>();
-	Transactor.CloseTransaction();
-
 	// Create the new Geometry Collection actor
 	AGeometryCollectionActor* NewActor = Cast<AGeometryCollectionActor>(AddActor(GetSelectedLevel(), AGeometryCollectionActor::StaticClass()));
 	check(NewActor->GetGeometryCollectionComponent());
