Fracture editor : add number of convex to the level statistics

diff --git a/Engine/Plugins/Experimental/ChaosEditor/Source/FractureEditor/Private/FractureEditorModeToolkit.cpp b/Engine/Plugins/Experimental/ChaosEditor/Source/FractureEditor/Private/FractureEditorModeToolkit.cpp
--- a/Engine/Plugins/Experimental/ChaosEditor/Source/FractureEditor/Private/FractureEditorModeToolkit.cpp
+++ b/Engine/Plugins/Experimental/ChaosEditor/Source/FractureEditor/Private/FractureEditorModeToolkit.cpp
@@ -58,6 +58,7 @@
 #include "FractureSelectionTools.h"
 #include "GeometryCollection/GeometryCollectionAlgo.h"
 #include "GeometryCollection/GeometryCollectionClusteringUtility.h"
+#include "GeometryCollection/Facades/CollectionHierarchyFacade.h"
 #include "Editor.h"
 #include "LevelEditorViewport.h"
 #include "Widgets/Layout/SUniformGridPanel.h"
@@ -2123,7 +2124,8 @@ void FFractureEditorModeToolkit::GetStatisticsSummary(FGeometryCollectionStatist
 
 	if (GeometryCollectionArray.Num() > 0)
 	{
-		TArray<int32> LevelTransformsAll;
+		TArray<uint32> TransformCountPerLevel;
+		TArray<uint32> ConvexCountPerLevel;
 		int32 LevelMax = INT_MIN;
 		int32 EmbeddedCount = 0;
 
@@ -2133,13 +2135,14 @@ void FFractureEditorModeToolkit::GetStatisticsSummary(FGeometryCollectionStatist
 
 			check(GeometryCollection);
 
-			if(GeometryCollection->HasAttribute("Level", FGeometryCollection::TransformGroup))
+			if(const TManagedArray<int32>* Levels = GeometryCollection->FindAttribute<int32>("Level", FGeometryCollection::TransformGroup))
 			{
-				const TManagedArray<int32>& Levels = GeometryCollection->GetAttribute<int32>("Level", FGeometryCollection::TransformGroup);
+
+				// num transforms per level
 				const TManagedArray<int32>& SimulationType = GeometryCollection->SimulationType;
 
 				TArray<int32> LevelTransforms;
-				for(int32 Element = 0, NumElement = Levels.Num(); Element < NumElement; ++Element)
+				for(int32 Element = 0, NumElement = Levels->Num(); Element < NumElement; ++Element)
 				{
 					if (SimulationType[Element] == FGeometryCollection::ESimulationTypes::FST_None)
 					{
@@ -2147,7 +2150,7 @@ void FFractureEditorModeToolkit::GetStatisticsSummary(FGeometryCollectionStatist
 					}
 					else
 					{
-						const int32 NodeLevel = Levels[Element];
+						const int32 NodeLevel = (*Levels)[Element];
 						if (LevelTransforms.Num() <= NodeLevel)
 						{
 							LevelTransforms.SetNumZeroed(NodeLevel + 1);
@@ -2156,26 +2159,76 @@ void FFractureEditorModeToolkit::GetStatisticsSummary(FGeometryCollectionStatist
 					}		
 				}
 
-				if (LevelTransformsAll.Num() < LevelTransforms.Num())
+				if (TransformCountPerLevel.Num() < LevelTransforms.Num())
 				{
-					LevelTransformsAll.SetNumZeroed(LevelTransforms.Num());
+					TransformCountPerLevel.SetNumZeroed(LevelTransforms.Num());
 				}
 				for(int32 Level = 0; Level < LevelTransforms.Num(); ++Level)
 				{
-					LevelTransformsAll[Level] += LevelTransforms[Level];
+					TransformCountPerLevel[Level] += LevelTransforms[Level];
 				}
 
 				if(LevelTransforms.Num() > LevelMax)
 				{
 					LevelMax = LevelTransforms.Num();
 				}
+
+				// convex per level
+				// todo(chaos) : should probably share some of this logic with the outliner 
+				TManagedArrayAccessor<TSet<int32>> GCTransformToConvexIndicesAttribute(*GeometryCollection, "TransformToConvexIndices", FGeometryCollection::TransformGroup);
+				if (GCTransformToConvexIndicesAttribute.IsValid())
+				{
+					const TManagedArray<TSet<int32>>& GCTransformToConvexIndices = GCTransformToConvexIndicesAttribute.Get();
+
+					Chaos::Facades::FCollectionHierarchyFacade HierarchyFacade(*GeometryCollection);
+					const TArray<int32> TransformIndices = HierarchyFacade.GetTransformArrayInDepthFirstOrder();
+
+					TArray<int32> ConvexCountArray;
+					ConvexCountArray.SetNumZeroed(TransformIndices.Num());
+
+					int32 MaxLevel = 0;
+					for (int32 TransformIndex : TransformIndices)
+					{
+						const int32 ConvexCount = GCTransformToConvexIndices[TransformIndex].Num();
+						if (ConvexCount > 0)
+						{
+							ConvexCountArray[TransformIndex] = ConvexCount;
+						}
+						// if count == 0 then we have already accumulated the children in the parents, no need to do anything
+						// so now just pass it to the direct parent ( we parse the index in a depth first manner ) 
+						const int32 ParentTransformIndex = GeometryCollection->Parent[TransformIndex];
+						if (ParentTransformIndex != INDEX_NONE)
+						{
+							// if parent has no convex then it will be a union of the aggregated children
+							const int32 ParentConvexCount = GCTransformToConvexIndices[ParentTransformIndex].Num();
+							if (ParentConvexCount == 0)
+							{
+								ConvexCountArray[ParentTransformIndex] += ConvexCountArray[TransformIndex];
+							}
+						}
+
+						const int32 Level = (*Levels)[TransformIndex];
+						MaxLevel = FMath::Max(MaxLevel, Level);
+					}
+					ConvexCountPerLevel.SetNumZeroed(MaxLevel + 1);
+					for (int32 TransformIndex = 0; TransformIndex < ConvexCountArray.Num(); TransformIndex++)
+					{
+						const int32 Level = (*Levels)[TransformIndex];
+						ConvexCountPerLevel[Level] += ConvexCountArray[TransformIndex];
+					}
+				}
 			}
 		}
 
+
+
+
 		Stats.CountsPerLevel.Reset();
+		Stats.ConvexCountPerLevel.Reset();
 		for (int32 Level = 0; Level < LevelMax; ++Level)
 		{
-			Stats.CountsPerLevel.Add(LevelTransformsAll[Level]);
+			Stats.CountsPerLevel = TransformCountPerLevel;
+			Stats.ConvexCountPerLevel = ConvexCountPerLevel;
 		}
 
 		Stats.EmbeddedCount = EmbeddedCount;
diff --git a/Engine/Plugins/Experimental/ChaosEditor/Source/FractureEditor/Private/SGeometryCollectionStatistics.cpp b/Engine/Plugins/Experimental/ChaosEditor/Source/FractureEditor/Private/SGeometryCollectionStatistics.cpp
--- a/Engine/Plugins/Experimental/ChaosEditor/Source/FractureEditor/Private/SGeometryCollectionStatistics.cpp
+++ b/Engine/Plugins/Experimental/ChaosEditor/Source/FractureEditor/Private/SGeometryCollectionStatistics.cpp
@@ -11,14 +11,16 @@
 namespace GeometryCollectionStatisticsUI
 {
 	static const FName NameLabel(TEXT("Stat"));
-	static const FName ValueLabel(TEXT("Value"));
+	static const FName NumBonesLabel(TEXT("NumBones"));
+	static const FName NumConvexLabel(TEXT("NumConvex"));
 }
 
 class SGeometryCollectionStatisticsRow : public SMultiColumnTableRow<TSharedPtr<int32>>
 {
 	SLATE_BEGIN_ARGS(SGeometryCollectionStatisticsRow) {}
 		SLATE_ARGUMENT(FText, Name)
-		SLATE_ARGUMENT(FText, Value)
+		SLATE_ARGUMENT(FText, NumBones)
+		SLATE_ARGUMENT(FText, NumConvex)
 		SLATE_ARGUMENT(FSlateColor, TextColor)
 	SLATE_END_ARGS()
 
@@ -32,7 +34,8 @@ public:
 	void Construct(const FArguments& InArgs, const TSharedRef<STableViewBase>& InOwnerTableView)
 	{
 		Name = InArgs._Name;
-		Value = InArgs._Value;
+		NumBones = InArgs._NumBones;
+		NumConvex = InArgs._NumConvex;
 		TextColor = InArgs._TextColor;
 		SMultiColumnTableRow<TSharedPtr<int32> >::Construct(FSuperRowType::FArguments(), InOwnerTableView);
 	}
@@ -48,9 +51,13 @@ public:
 		{
 			ColumnText = Name;
 		}
-		else if (ColumnName == GeometryCollectionStatisticsUI::ValueLabel)
+		else if (ColumnName == GeometryCollectionStatisticsUI::NumBonesLabel)
 		{
-			ColumnText = Value;
+			ColumnText = NumBones;
+		}
+		else if (ColumnName == GeometryCollectionStatisticsUI::NumConvexLabel)
+		{
+			ColumnText = NumConvex;
 		}
 
 		return SNew(SBox)
@@ -65,12 +72,9 @@ public:
 
 private:
 
-	/** Name of the stat. */
 	FText Name;
-
-	/** Value of the stat. */
-	FText Value;
-
+	FText NumBones;
+	FText NumConvex;
 	FSlateColor TextColor;
 };
 
@@ -86,8 +90,11 @@ void SGeometryCollectionStatistics::Construct(const FArguments& InArgs)
 				SNew(SHeaderRow)
 				+SHeaderRow::Column(GeometryCollectionStatisticsUI::NameLabel)
 				.DefaultLabel(LOCTEXT("NameColumnHeaderName", "Name"))
-				+SHeaderRow::Column(GeometryCollectionStatisticsUI::ValueLabel)
-				.DefaultLabel(LOCTEXT("ValueColumnHeaderName", "Value"))
+				+SHeaderRow::Column(GeometryCollectionStatisticsUI::NumBonesLabel)
+				.DefaultLabel(LOCTEXT("ValueColumnHeaderName", "Num Bones"))
+				+ SHeaderRow::Column(GeometryCollectionStatisticsUI::NumConvexLabel)
+				.DefaultLabel(LOCTEXT("ValueColumnHeaderName", "Num Convex"))
+
 			)
 		];
 }
@@ -113,6 +120,7 @@ void SGeometryCollectionStatistics::SetStatistics(const FGeometryCollectionStati
 		{
 			FText::Format(LOCTEXT("Level", "Level {0}"), FText::AsNumber(Level)),
 			FText::AsNumber(Stats.CountsPerLevel[Level]),
+			FText::AsNumber(Stats.ConvexCountPerLevel[Level]),
 			TextColor
 		};
 		StatItems.Emplace(MakeShared<FStatRow>(MoveTemp(Row)));
@@ -122,6 +130,7 @@ void SGeometryCollectionStatistics::SetStatistics(const FGeometryCollectionStati
 	{
 		FText(LOCTEXT("Embedded", "Embedded")),
 		FText::AsNumber(Stats.EmbeddedCount),
+		FText::AsNumber(0),
 		FSlateColor::UseForeground()
 	};
 	StatItems.Emplace(MakeShared<FStatRow>(MoveTemp(Row)));
@@ -134,7 +143,8 @@ TSharedRef<ITableRow> SGeometryCollectionStatistics::HandleGenerateRow(FListItem
 {	
 	return SNew(SGeometryCollectionStatisticsRow, OwnerTable)
 		.Name(InItem->Name)
-		.Value(InItem->Value)
+		.NumBones(InItem->NumBones)
+		.NumConvex(InItem->NumConvex)
 		.TextColor(InItem->TextColor);
 }
 
diff --git a/Engine/Plugins/Experimental/ChaosEditor/Source/FractureEditor/Private/SGeometryCollectionStatistics.h b/Engine/Plugins/Experimental/ChaosEditor/Source/FractureEditor/Private/SGeometryCollectionStatistics.h
--- a/Engine/Plugins/Experimental/ChaosEditor/Source/FractureEditor/Private/SGeometryCollectionStatistics.h
+++ b/Engine/Plugins/Experimental/ChaosEditor/Source/FractureEditor/Private/SGeometryCollectionStatistics.h
@@ -10,6 +10,7 @@
 struct FGeometryCollectionStatistics
 {
 	TArray<uint32> CountsPerLevel;
+	TArray<uint32> ConvexCountPerLevel;
 	uint32 EmbeddedCount = 0;
 };
 
@@ -37,7 +38,8 @@ private:
 	struct FStatRow
 	{
 		FText Name;
-		FText Value;
+		FText NumBones;
+		FText NumConvex;
 		FSlateColor TextColor;
 	};
 
