Fix SGeometryCollectionStatistics localization text conflict

diff --git a/Engine/Plugins/Experimental/ChaosEditor/Source/FractureEditor/Private/SGeometryCollectionStatistics.cpp b/Engine/Plugins/Experimental/ChaosEditor/Source/FractureEditor/Private/SGeometryCollectionStatistics.cpp
--- a/Engine/Plugins/Experimental/ChaosEditor/Source/FractureEditor/Private/SGeometryCollectionStatistics.cpp
+++ b/Engine/Plugins/Experimental/ChaosEditor/Source/FractureEditor/Private/SGeometryCollectionStatistics.cpp
@@ -91,9 +91,9 @@ void SGeometryCollectionStatistics::Construct(const FArguments& InArgs)
 				+SHeaderRow::Column(GeometryCollectionStatisticsUI::NameLabel)
 				.DefaultLabel(LOCTEXT("NameColumnHeaderName", "Name"))
 				+SHeaderRow::Column(GeometryCollectionStatisticsUI::NumBonesLabel)
-				.DefaultLabel(LOCTEXT("ValueColumnHeaderName", "Num Bones"))
+				.DefaultLabel(LOCTEXT("NumBonesColumnHeaderName", "Num Bones"))
 				+ SHeaderRow::Column(GeometryCollectionStatisticsUI::NumConvexLabel)
-				.DefaultLabel(LOCTEXT("ValueColumnHeaderName", "Num Convex"))
+				.DefaultLabel(LOCTEXT("NumConvexColumnHeaderName", "Num Convex"))
 
 			)
 		];
