Fracture editor : make histogram SetInspectedAttribute method more resilkient to ill formed collections

diff --git a/Engine/Plugins/Experimental/ChaosEditor/Source/FractureEditor/Private/SGeometryCollectionHistogram.cpp b/Engine/Plugins/Experimental/ChaosEditor/Source/FractureEditor/Private/SGeometryCollectionHistogram.cpp
--- a/Engine/Plugins/Experimental/ChaosEditor/Source/FractureEditor/Private/SGeometryCollectionHistogram.cpp
+++ b/Engine/Plugins/Experimental/ChaosEditor/Source/FractureEditor/Private/SGeometryCollectionHistogram.cpp
@@ -57,42 +57,49 @@ void FGeometryCollectionHistogramItem::SetInspectedAttribute(EInspectedAttribute
 
 		if (Component)
 		{
-			FGeometryCollection* Collection = Component->GetRestCollection()->GetGeometryCollection().Get();
-
-			switch (InspectedAttribute)
+			if (FGeometryCollection* Collection = Component->GetRestCollection()->GetGeometryCollection().Get())
 			{
-			case EInspectedAttributeEnum::Volume:
-				InspectedValue = Collection->GetAttribute<float>(TEXT("Volume"), TEXT("Transform"))[BoneIndex];
-				HoverString = FString::Printf(TEXT("%d: %.2f"), BoneIndex, InspectedValue);
-				break;
-
-			case EInspectedAttributeEnum::Level:
-				{
-					int32 Level = Collection->GetAttribute<int32>(TEXT("Level"), TEXT("Transform"))[BoneIndex];
-					HoverString = FString::Printf(TEXT("%d: %d"), BoneIndex, Level);
-					InspectedValue = static_cast<float>(Level);
-				}
-				break;
 
-			case EInspectedAttributeEnum::InitialDynamicState:
+				switch (InspectedAttribute)
 				{
-					int32 InitialDynamicState = static_cast<int32>(Collection->GetAttribute<int32>(TEXT("InitialDynamicState"), TEXT("Transform"))[BoneIndex]);
-					
-					static const TArray<FString> HoverNames{ "No Override", "Sleeping", "Kinematic", "Static" };
-					HoverString = FString::Printf(TEXT("%d: %s"), BoneIndex, *HoverNames[InitialDynamicState]);
+				case EInspectedAttributeEnum::Volume:
+					{
+						const TManagedArrayAccessor<float> VolumeAttribute(*Collection, "Volume", FGeometryCollection::TransformAttribute);
+						InspectedValue = (VolumeAttribute.IsValid() && VolumeAttribute.IsValidIndex(BoneIndex)) ? VolumeAttribute[BoneIndex] : 0.0f;
+						HoverString = FString::Printf(TEXT("%d: %.2f"), BoneIndex, InspectedValue);
+					}
+					break;
+				case EInspectedAttributeEnum::Level:
+					{
+						const TManagedArrayAccessor<int32> LevelAttribute(*Collection, "Level", FGeometryCollection::TransformAttribute);
+						const int32 Level = (LevelAttribute.IsValid() && LevelAttribute.IsValidIndex(BoneIndex)) ? LevelAttribute[BoneIndex] : 0;
+						HoverString = FString::Printf(TEXT("%d: %d"), BoneIndex, Level);
+						InspectedValue = static_cast<float>(Level);
+					}
+					break;
 
-					InspectedValue = static_cast<float>(InitialDynamicState);
-				}
-				break;
+				case EInspectedAttributeEnum::InitialDynamicState:
+					{
+						const TManagedArrayAccessor<int32> InitialDynamicStateAttribute(*Collection, "InitialDynamicState", FGeometryCollection::TransformAttribute);
+						const int32 InitialDynamicState = (InitialDynamicStateAttribute.IsValid() && InitialDynamicStateAttribute.IsValidIndex(BoneIndex)) ? InitialDynamicStateAttribute[BoneIndex] : 0;
+						static const TArray<FString> HoverNames{ "No Override", "Sleeping", "Kinematic", "Static" };
+						HoverString = FString::Printf(TEXT("%d: %s"), BoneIndex, *HoverNames[InitialDynamicState]);
+						InspectedValue = static_cast<float>(InitialDynamicState);
+					}
+					break;
 
-			case EInspectedAttributeEnum::Size:
-				InspectedValue = Collection->GetAttribute<float>(TEXT("Size"), TEXT("Transform"))[BoneIndex];
-				HoverString = FString::Printf(TEXT("%d: %.2f"), BoneIndex, InspectedValue);
-				break;
+				case EInspectedAttributeEnum::Size:
+					{
+						const TManagedArrayAccessor<float> SizeAttribute(*Collection, "Size", FGeometryCollection::TransformAttribute);
+						InspectedValue = (SizeAttribute.IsValid() && SizeAttribute.IsValidIndex(BoneIndex)) ? SizeAttribute[BoneIndex] : 0.0f;
+						HoverString = FString::Printf(TEXT("%d: %.2f"), BoneIndex, InspectedValue);
+					}
+					break;
 
-			default:
-				// Invalid inspection attribute
-				check(false);
+				default:
+					// Invalid inspection attribute
+					check(false);
+				}
 			}
 		}	
 	}
