Fracture Tool : Fix potential crash when generatic per level statistic
- also remove the unecessary loop over setting the arrays

diff --git a/Engine/Plugins/Experimental/ChaosEditor/Source/FractureEditor/Private/FractureEditorModeToolkit.cpp b/Engine/Plugins/Experimental/ChaosEditor/Source/FractureEditor/Private/FractureEditorModeToolkit.cpp
--- a/Engine/Plugins/Experimental/ChaosEditor/Source/FractureEditor/Private/FractureEditorModeToolkit.cpp
+++ b/Engine/Plugins/Experimental/ChaosEditor/Source/FractureEditor/Private/FractureEditorModeToolkit.cpp
@@ -2220,17 +2220,8 @@ void FFractureEditorModeToolkit::GetStatisticsSummary(FGeometryCollectionStatist
 			}
 		}
 
-
-
-
-		Stats.CountsPerLevel.Reset();
-		Stats.ConvexCountPerLevel.Reset();
-		for (int32 Level = 0; Level < LevelMax; ++Level)
-		{
-			Stats.CountsPerLevel = TransformCountPerLevel;
-			Stats.ConvexCountPerLevel = ConvexCountPerLevel;
-		}
-
+		Stats.CountsPerLevel = TransformCountPerLevel;
+		Stats.ConvexCountPerLevel = ConvexCountPerLevel;
 		Stats.EmbeddedCount = EmbeddedCount;
 	}
 }
diff --git a/Engine/Plugins/Experimental/ChaosEditor/Source/FractureEditor/Private/SGeometryCollectionStatistics.cpp b/Engine/Plugins/Experimental/ChaosEditor/Source/FractureEditor/Private/SGeometryCollectionStatistics.cpp
--- a/Engine/Plugins/Experimental/ChaosEditor/Source/FractureEditor/Private/SGeometryCollectionStatistics.cpp
+++ b/Engine/Plugins/Experimental/ChaosEditor/Source/FractureEditor/Private/SGeometryCollectionStatistics.cpp
@@ -120,7 +120,7 @@ void SGeometryCollectionStatistics::SetStatistics(const FGeometryCollectionStati
 		{
 			FText::Format(LOCTEXT("Level", "Level {0}"), FText::AsNumber(Level)),
 			FText::AsNumber(Stats.CountsPerLevel[Level]),
-			FText::AsNumber(Stats.ConvexCountPerLevel[Level]),
+			FText::AsNumber(Stats.ConvexCountPerLevel.IsValidIndex(Level)? Stats.ConvexCountPerLevel[Level]: 0),
 			TextColor
 		};
 		StatItems.Emplace(MakeShared<FStatRow>(MoveTemp(Row)));
