guard against crash due to geometry collection outliner items w/ invalid bone indices

diff --git a/Engine/Plugins/Experimental/ChaosEditor/Source/FractureEditor/Private/SGeometryCollectionOutliner.cpp b/Engine/Plugins/Experimental/ChaosEditor/Source/FractureEditor/Private/SGeometryCollectionOutliner.cpp
--- a/Engine/Plugins/Experimental/ChaosEditor/Source/FractureEditor/Private/SGeometryCollectionOutliner.cpp
+++ b/Engine/Plugins/Experimental/ChaosEditor/Source/FractureEditor/Private/SGeometryCollectionOutliner.cpp
@@ -192,6 +192,11 @@ UOutlinerSettings::UOutlinerSettings(const FObjectInitializer& ObjInit)
 
 TSharedRef<SWidget> SGeometryCollectionOutlinerRow::GenerateWidgetForColumn(const FName& ColumnName)
 {
+	if (!ensure(Item->IsValidBone()))
+	{
+		return Item->MakeEmptyColumnWidget();
+	}
+
 	if (ColumnName == SGeometryCollectionOutlinerColumnID::BoneIndex)
 	{
 		const TSharedPtr<SWidget> NameWidget = Item->MakeBoneIndexColumnWidget();
@@ -649,6 +654,11 @@ static T GetAttributeValue(const TManagedArrayAccessor<T>& Attribute, int32 Inde
 	return (Attribute.IsValid()) ? Attribute.Get()[Index] : Default;
 }
 
+bool FGeometryCollectionItemDataFacade::IsValidBoneIndex(int32 BoneIndex) const
+{
+	return BoneIndex >= 0 && BoneIndex < BoneNameAttribute.Num();
+}
+
 FString FGeometryCollectionItemDataFacade::GetBoneName(int32 Index) const
 {
 	return GetAttributeValue(BoneNameAttribute, Index, FString());
@@ -1116,6 +1126,12 @@ TSharedRef<ITableRow> FGeometryCollectionTreeItemBone::MakeTreeRowWidget(const T
 	return SNew(SGeometryCollectionOutlinerRow, InOwnerTable, SharedThis(this));
 }
 
+bool FGeometryCollectionTreeItemBone::IsValidBone() const
+{
+	const FGeometryCollectionItemDataFacade& DataCollectionFacade = GetDataCollectionFacade();
+	return DataCollectionFacade.IsValidBoneIndex(BoneIndex);
+}
+
 TSharedRef<SWidget> FGeometryCollectionTreeItemBone::MakeBoneIndexColumnWidget() const
 {
 	return SNew(SHorizontalBox)
diff --git a/Engine/Plugins/Experimental/ChaosEditor/Source/FractureEditor/Private/SGeometryCollectionOutliner.h b/Engine/Plugins/Experimental/ChaosEditor/Source/FractureEditor/Private/SGeometryCollectionOutliner.h
--- a/Engine/Plugins/Experimental/ChaosEditor/Source/FractureEditor/Private/SGeometryCollectionOutliner.h
+++ b/Engine/Plugins/Experimental/ChaosEditor/Source/FractureEditor/Private/SGeometryCollectionOutliner.h
@@ -98,6 +98,7 @@ public:
 	const TManagedArray<int32>& GetSimulationType() const { return SimulationTypeAttribute.Get(); }
 	const TManagedArray<bool>& GetVisible() const { return VisibleAttribute.Get(); }
 
+	bool IsValidBoneIndex(int32 BoneIndex) const;
 	FString GetBoneName(int32 Index) const;
 	float GetRelativeSize(int32 Index) const;
 	float GetVolumetricUnit(int32 Index) const;
@@ -210,6 +211,7 @@ public:
 	TSharedRef<SWidget> MakeConvexCountColumnWidget() const;
 	TSharedRef<SWidget> MakeEmptyColumnWidget() const;
 	virtual void GetChildren(FGeometryCollectionTreeItemList& OutChildren) override;
+	bool IsValidBone() const;
 	virtual int32 GetBoneIndex() const override { return BoneIndex; }
 	virtual UGeometryCollectionComponent* GetComponent() const { return ParentComponentItem->GetComponent(); }
 	bool HasChildren() const;
