Move the fracture 'magnet cluster' tool's logic into the fracture engine plugin and make a dataflow node with the same functionality

diff --git a/Engine/Plugins/Experimental/ChaosEditor/Source/FractureEditor/Private/FractureToolClusterMagnet.cpp b/Engine/Plugins/Experimental/ChaosEditor/Source/FractureEditor/Private/FractureToolClusterMagnet.cpp
--- a/Engine/Plugins/Experimental/ChaosEditor/Source/FractureEditor/Private/FractureToolClusterMagnet.cpp
+++ b/Engine/Plugins/Experimental/ChaosEditor/Source/FractureEditor/Private/FractureToolClusterMagnet.cpp
@@ -16,6 +16,8 @@
 #include "GeometryCollection/GeometryCollectionClusteringUtility.h"
 #include "GeometryCollection/GeometryCollectionProximityUtility.h"
 
+#include "FractureEngineClustering.h"
+
 #include "GeometryCollection/GeometryCollectionComponent.h"
 
 #include UE_INLINE_GENERATED_CPP_BY_NAME(FractureToolClusterMagnet)
@@ -79,69 +81,14 @@ void UFractureToolClusterMagnet::Execute(TWeakPtr<FFractureEditorModeToolkit> In
 
 		for (FFractureToolContext& Context : Contexts)
 		{
-			// We require certain attributes present to proceed.
-			if (!CheckPresenceOfNecessaryAttributes(Context.GetGeometryCollection()))
-			{
-				continue;
-			}
-
 			FGeometryCollectionEdit Edit(Context.GetGeometryCollectionComponent(), GeometryCollection::EEditUpdate::RestPhysicsDynamic);
 
-			const TManagedArray<TSet<int32>>& Children = Context.GetGeometryCollection()->Children;
-			if (!Context.GetGeometryCollection()->HasAttribute("Level", FTransformCollection::TransformGroup))
-			{
-				FGeometryCollectionClusteringUtility::UpdateHierarchyLevelOfChildren(Edit.GetRestCollection()->GetGeometryCollection().Get(), -1);
-			}
-			const TManagedArray<int32>& Levels = Context.GetGeometryCollection()->GetAttribute<int32>("Level", FGeometryCollection::TransformGroup);
-			int32 StartTransformCount = Children.Num();
-			
-			Context.Sanitize();
-			TMap<int32, TArray<int32>> ClusteredSelection = Context.GetClusteredSelections();
-
-			for (TPair<int32, TArray<int32>>& Group : ClusteredSelection)
-			{
-				if (Group.Key == INDEX_NONE) // Group is top level
-				{
-					continue;
-				}
-
-				// We have the connections for the leaf nodes of our geometry collection. We want to percolate those up to the top nodes.
-				TMap<int32, TSet<int32>> TopNodeConnectivity = InitializeConnectivity(Children[Group.Key], Context.GetGeometryCollection(), Levels[Group.Key]+1);
-
-				// Separate the top nodes into cluster magnets and a pool of available nodes.
-				TArray<FClusterMagnet> ClusterMagnets;
-				TSet<int32> RemainingPool;
-				SeparateClusterMagnets(Children[Group.Key], Group.Value, TopNodeConnectivity, ClusterMagnets, RemainingPool);
-
-				for (uint32 Iteration = 0; Iteration < ClusterMagnetSettings->Iterations; ++Iteration)
-				{
-					bool bNeighborsAbsorbed = false;
-
-					// each cluster gathers adjacent nodes from the pool
-					for (FClusterMagnet& ClusterMagnet : ClusterMagnets)
-					{
-						bNeighborsAbsorbed |= AbsorbClusterNeighbors(TopNodeConnectivity, ClusterMagnet, RemainingPool);
-					}
-
-					// early termination
-					if (!bNeighborsAbsorbed)
-					{
-						break;
-					}
-				}
-
-				// Create new clusters from the cluster magnets
-				for (const FClusterMagnet& ClusterMagnet : ClusterMagnets)
-				{
-					if (ClusterMagnet.ClusteredNodes.Num() > 1)
-					{
-						TArray<int32> NewChildren = ClusterMagnet.ClusteredNodes.Array();
-						NewChildren.Sort();
-						FGeometryCollectionClusteringUtility::ClusterBonesUnderNewNode(Context.GetGeometryCollection().Get(), NewChildren[0], NewChildren, false, false);
-
-					}
-				}
-			}
+			TArray<int32> Selection = Context.GetSelection();
+			int32 StartTransformCount = Context.GetGeometryCollection()->NumElements(FGeometryCollection::TransformGroup);
+			FFractureEngineClustering::ClusterMagnet(
+				*Context.GetGeometryCollection(),
+				Selection,
+				ClusterMagnetSettings->Iterations);
 
 			Context.GenerateGuids(StartTransformCount);
 
@@ -152,117 +99,5 @@ void UFractureToolClusterMagnet::Execute(TWeakPtr<FFractureEditorModeToolkit> In
 	}
 }
 
-bool UFractureToolClusterMagnet::CheckPresenceOfNecessaryAttributes(const FGeometryCollectionPtr GeometryCollection) const
-{
-	if (!GeometryCollection->HasAttribute("Level", FTransformCollection::TransformGroup))
-	{
-		UE_LOG(LogFractureTool, Error, TEXT("Cannot execute Cluster Magnet tool: missing Level attribute."));
-		return false;
-	}
-
-	if (!GeometryCollection->HasAttribute("Proximity", FGeometryCollection::GeometryGroup))
-	{
-		UE_LOG(LogFractureTool, Error, TEXT("Cannot execute Cluster Magnet tool: missing Proximity attribute."));
-		return false;
-	}
-
-	return true;
-}
-
-TMap<int32, TSet<int32>> UFractureToolClusterMagnet::InitializeConnectivity(const TSet<int32>& TopNodes, FGeometryCollectionPtr GeometryCollection, int32 OperatingLevel) const
-{
-	FGeometryCollectionProximityUtility ProximityUtility(GeometryCollection.Get());
-	ProximityUtility.RequireProximity();
-
-	TMap<int32, TSet<int32>> ConnectivityMap;
-	for (int32 Index : TopNodes)
-	{
-		// Collect the proximity indices of all the leaf nodes under this top node,
-		// traced back up to its parent top node, so that all connectivity describes
-		// relationships only between top nodes.
-		TSet<int32> Connections;
-		CollectTopNodeConnections(GeometryCollection, Index, OperatingLevel, Connections);
-		Connections.Remove(Index);
-
-		// Remove any connections outside the current operating branch.
-		ConnectivityMap.Add(Index, Connections.Intersect(TopNodes));
-	}
-
-	return ConnectivityMap;
-}
-
-void UFractureToolClusterMagnet::CollectTopNodeConnections(FGeometryCollectionPtr GeometryCollection, int32 Index, int32 OperatingLevel, TSet<int32>& OutConnections) const
-{
-	const TManagedArray<int32>& TransformToGeometryIndex = GeometryCollection->TransformToGeometryIndex;
-	if (GeometryCollection->SimulationType[Index] == FGeometryCollection::ESimulationTypes::FST_Rigid
-		&& TransformToGeometryIndex[Index] != INDEX_NONE) // rigid node with geometry, leaf of the simulated part
-	{
-		const TManagedArray<TSet<int32>>& Proximity = GeometryCollection->GetAttribute<TSet<int32>>("Proximity", FGeometryCollection::GeometryGroup);
-		const TManagedArray<int32>& GeometryToTransformIndex = GeometryCollection->TransformIndex;
-
-
-		for (int32 Neighbor : Proximity[TransformToGeometryIndex[Index]])
-		{
-			int32 NeighborTransformIndex = GeometryToTransformIndex[Neighbor];
-			OutConnections.Add(FGeometryCollectionClusteringUtility::GetParentOfBoneAtSpecifiedLevel(GeometryCollection.Get(), NeighborTransformIndex, OperatingLevel));
-		}		
-	}
-	else
-	{
-		const TManagedArray<TSet<int32>>& Children = GeometryCollection->Children;
-		for (int32 ChildIndex : Children[Index])
-		{
-			CollectTopNodeConnections(GeometryCollection, ChildIndex, OperatingLevel, OutConnections);
-		}
-	}
-}
-
-void UFractureToolClusterMagnet::SeparateClusterMagnets(
-	const TSet<int32>& TopNodes,
-	const TArray<int32>& Selection,
-	const TMap<int32, TSet<int32>>& TopNodeConnectivity,
-	TArray<FClusterMagnet>& OutClusterMagnets,
-	TSet<int32>& OutRemainingPool) const
-{
-	OutClusterMagnets.Reserve(TopNodes.Num());
-	OutRemainingPool.Reserve(TopNodes.Num());
-	
-	for (int32 Index : TopNodes)
-	{
-		if (Selection.Contains(Index))
-		{
-			OutClusterMagnets.AddDefaulted();
-			FClusterMagnet& NewMagnet = OutClusterMagnets.Last();
-			NewMagnet.ClusteredNodes.Add(Index);
-			NewMagnet.Connections = TopNodeConnectivity[Index];
-		}
-		else
-		{
-			OutRemainingPool.Add(Index);
-		}
-	}
-}
-
-bool UFractureToolClusterMagnet::AbsorbClusterNeighbors(const TMap<int32, TSet<int32>> TopNodeConnectivity, FClusterMagnet& OutClusterMagnet, TSet<int32>& OutRemainingPool) const
-{
-	// Return true if neighbors were absorbed.
-	bool bNeighborsAbsorbed = false;
-
-	TSet<int32> NewConnections;
-	for (int32 NeighborIndex : OutClusterMagnet.Connections)
-	{
-		// If the neighbor is still in the pool, absorb it and its connections.
-		if (OutRemainingPool.Contains(NeighborIndex))
-		{
-			OutClusterMagnet.ClusteredNodes.Add(NeighborIndex);
-			NewConnections.Append(TopNodeConnectivity[NeighborIndex]);
-			OutRemainingPool.Remove(NeighborIndex);
-			bNeighborsAbsorbed = true;
-		}
-	}
-	OutClusterMagnet.Connections.Append(NewConnections);
-
-	return bNeighborsAbsorbed;
-}
 
 #undef LOCTEXT_NAMESPACE
diff --git a/Engine/Plugins/Experimental/ChaosEditor/Source/FractureEditor/Private/FractureToolClusterMagnet.h b/Engine/Plugins/Experimental/ChaosEditor/Source/FractureEditor/Private/FractureToolClusterMagnet.h
--- a/Engine/Plugins/Experimental/ChaosEditor/Source/FractureEditor/Private/FractureToolClusterMagnet.h
+++ b/Engine/Plugins/Experimental/ChaosEditor/Source/FractureEditor/Private/FractureToolClusterMagnet.h
@@ -14,12 +14,6 @@ namespace Chaos
 }
 
 
-struct FClusterMagnet
-{
-	TSet<int32> ClusteredNodes;
-	TSet<int32> Connections;
-};
-
 
 UCLASS(DisplayName = "Cluster Magnet", Category = "FractureTools")
 class UFractureClusterMagnetSettings : public UFractureToolSettings
@@ -62,11 +56,4 @@ public:
 	UPROPERTY(EditAnywhere, Category = AutoCluster)
 	TObjectPtr<UFractureClusterMagnetSettings> ClusterMagnetSettings;
 
-private:
-	bool CheckPresenceOfNecessaryAttributes(const FGeometryCollectionPtr GeometryCollection) const;
-	TMap<int32, TSet<int32>> InitializeConnectivity(const TSet<int32>& TopNodes, FGeometryCollectionPtr GeometryCollection, int32 OperatingLevel) const;
-	void CollectTopNodeConnections(FGeometryCollectionPtr GeometryCollection, int32 Index, int32 OperatingLevel, TSet<int32>& OutConnections) const;
-	void SeparateClusterMagnets(const TSet<int32>& TopNodes, const TArray<int32>& Selection, const TMap<int32, TSet<int32>>& TopNodeConnectivity, TArray<FClusterMagnet>& OutClusterMagnets, TSet<int32>& OutRemainingPool) const;
-	bool AbsorbClusterNeighbors(const TMap<int32, TSet<int32>> TopNodeConnectivity, FClusterMagnet& OutClusterMagnets, TSet<int32>& OutRemainingPool) const;
-
 };
diff --git a/Engine/Plugins/Experimental/Fracture/Source/FractureEngine/Private/FractureEngineClustering.cpp b/Engine/Plugins/Experimental/Fracture/Source/FractureEngine/Private/FractureEngineClustering.cpp
--- a/Engine/Plugins/Experimental/Fracture/Source/FractureEngine/Private/FractureEngineClustering.cpp
+++ b/Engine/Plugins/Experimental/Fracture/Source/FractureEngine/Private/FractureEngineClustering.cpp
@@ -844,3 +844,181 @@ bool FFractureEngineClustering::MergeSelectedClusters(FGeometryCollection& Geome
 
 	return false;
 }
+
+namespace UE::Private::ClusterMagnet
+{
+	struct FClusterMagnet
+	{
+		TSet<int32> ClusteredNodes;
+		TSet<int32> Connections;
+	};
+
+	TMap<int32, TSet<int32>> InitializeConnectivity(const TSet<int32>& TopNodes, FGeometryCollection* GeometryCollection, int32 OperatingLevel);
+	void CollectTopNodeConnections(FGeometryCollection* GeometryCollection, int32 Index, int32 OperatingLevel, TSet<int32>& OutConnections);
+	void SeparateClusterMagnets(const TSet<int32>& TopNodes, const TArray<int32>& Selection, const TMap<int32, TSet<int32>>& TopNodeConnectivity, TArray<FClusterMagnet>& OutClusterMagnets, TSet<int32>& OutRemainingPool);
+	bool AbsorbClusterNeighbors(const TMap<int32, TSet<int32>> TopNodeConnectivity, FClusterMagnet& OutClusterMagnets, TSet<int32>& OutRemainingPool);
+
+
+	TMap<int32, TSet<int32>> InitializeConnectivity(const TSet<int32>& TopNodes, FGeometryCollection* GeometryCollection, int32 OperatingLevel)
+	{
+		FGeometryCollectionProximityUtility ProximityUtility(GeometryCollection);
+		ProximityUtility.RequireProximity();
+
+		TMap<int32, TSet<int32>> ConnectivityMap;
+		for (int32 Index : TopNodes)
+		{
+			// Collect the proximity indices of all the leaf nodes under this top node,
+			// traced back up to its parent top node, so that all connectivity describes
+			// relationships only between top nodes.
+			TSet<int32> Connections;
+			CollectTopNodeConnections(GeometryCollection, Index, OperatingLevel, Connections);
+			Connections.Remove(Index);
+
+			// Remove any connections outside the current operating branch.
+			ConnectivityMap.Add(Index, Connections.Intersect(TopNodes));
+		}
+
+		return ConnectivityMap;
+	}
+
+	void CollectTopNodeConnections(FGeometryCollection* GeometryCollection, int32 Index, int32 OperatingLevel, TSet<int32>& OutConnections)
+	{
+		const TManagedArray<int32>& TransformToGeometryIndex = GeometryCollection->TransformToGeometryIndex;
+		if (GeometryCollection->SimulationType[Index] == FGeometryCollection::ESimulationTypes::FST_Rigid
+			&& TransformToGeometryIndex[Index] != INDEX_NONE) // rigid node with geometry, leaf of the simulated part
+		{
+			const TManagedArray<TSet<int32>>& Proximity = GeometryCollection->GetAttribute<TSet<int32>>("Proximity", FGeometryCollection::GeometryGroup);
+			const TManagedArray<int32>& GeometryToTransformIndex = GeometryCollection->TransformIndex;
+
+
+			for (int32 Neighbor : Proximity[TransformToGeometryIndex[Index]])
+			{
+				int32 NeighborTransformIndex = GeometryToTransformIndex[Neighbor];
+				OutConnections.Add(FGeometryCollectionClusteringUtility::GetParentOfBoneAtSpecifiedLevel(GeometryCollection, NeighborTransformIndex, OperatingLevel));
+			}
+		}
+		else
+		{
+			const TManagedArray<TSet<int32>>& Children = GeometryCollection->Children;
+			for (int32 ChildIndex : Children[Index])
+			{
+				CollectTopNodeConnections(GeometryCollection, ChildIndex, OperatingLevel, OutConnections);
+			}
+		}
+	}
+
+	void SeparateClusterMagnets(
+		const TSet<int32>& TopNodes,
+		const TArray<int32>& Selection,
+		const TMap<int32, TSet<int32>>& TopNodeConnectivity,
+		TArray<FClusterMagnet>& OutClusterMagnets,
+		TSet<int32>& OutRemainingPool)
+	{
+		OutClusterMagnets.Reserve(TopNodes.Num());
+		OutRemainingPool.Reserve(TopNodes.Num());
+
+		for (int32 Index : TopNodes)
+		{
+			if (Selection.Contains(Index))
+			{
+				OutClusterMagnets.AddDefaulted();
+				FClusterMagnet& NewMagnet = OutClusterMagnets.Last();
+				NewMagnet.ClusteredNodes.Add(Index);
+				NewMagnet.Connections = TopNodeConnectivity[Index];
+			}
+			else
+			{
+				OutRemainingPool.Add(Index);
+			}
+		}
+	}
+
+	bool AbsorbClusterNeighbors(const TMap<int32, TSet<int32>> TopNodeConnectivity, FClusterMagnet& OutClusterMagnet, TSet<int32>& OutRemainingPool)
+	{
+		// Return true if neighbors were absorbed.
+		bool bNeighborsAbsorbed = false;
+
+		TSet<int32> NewConnections;
+		for (int32 NeighborIndex : OutClusterMagnet.Connections)
+		{
+			// If the neighbor is still in the pool, absorb it and its connections.
+			if (OutRemainingPool.Contains(NeighborIndex))
+			{
+				OutClusterMagnet.ClusteredNodes.Add(NeighborIndex);
+				NewConnections.Append(TopNodeConnectivity[NeighborIndex]);
+				OutRemainingPool.Remove(NeighborIndex);
+				bNeighborsAbsorbed = true;
+			}
+		}
+		OutClusterMagnet.Connections.Append(NewConnections);
+
+		return bNeighborsAbsorbed;
+	}
+}
+
+bool FFractureEngineClustering::ClusterMagnet(
+	FGeometryCollection& GeometryCollection,
+	TArray<int32>& InOutSelection,
+	int32 Iterations
+)
+{
+	using namespace UE::Private::ClusterMagnet;
+
+	Chaos::Facades::FCollectionHierarchyFacade HierarchyFacade(GeometryCollection);
+	HierarchyFacade.GenerateLevelAttribute();
+	const TManagedArray<int32>& Levels = GeometryCollection.GetAttribute<int32>("Level", FGeometryCollection::TransformGroup);
+
+	GeometryCollection::Facades::FCollectionTransformSelectionFacade SelectionFacade(GeometryCollection);
+	SelectionFacade.Sanitize(InOutSelection);
+	SelectionFacade.ConvertEmbeddedSelectionToParents(InOutSelection); // embedded geo must stay attached to parent
+
+	const TManagedArray<TSet<int32>>& Children = GeometryCollection.Children;
+	TMap<int32, TArray<int32>> ClusteredSelection = SelectionFacade.GetClusteredSelections(InOutSelection);
+	
+	for (TPair<int32, TArray<int32>>& Group : ClusteredSelection)
+	{
+		if (Group.Key == INDEX_NONE) // Group is top level
+		{
+			continue;
+		}
+
+		// We have the connections for the leaf nodes of our geometry collection. We want to percolate those up to the top nodes.
+		TMap<int32, TSet<int32>> TopNodeConnectivity = InitializeConnectivity(Children[Group.Key], &GeometryCollection, Levels[Group.Key]+1);
+
+		// Separate the top nodes into cluster magnets and a pool of available nodes.
+		TArray<FClusterMagnet> ClusterMagnets;
+		TSet<int32> RemainingPool;
+		SeparateClusterMagnets(Children[Group.Key], Group.Value, TopNodeConnectivity, ClusterMagnets, RemainingPool);
+
+		for (int32 Iteration = 0; Iteration < Iterations; ++Iteration)
+		{
+			bool bNeighborsAbsorbed = false;
+
+			// each cluster gathers adjacent nodes from the pool
+			for (FClusterMagnet& ClusterMagnet : ClusterMagnets)
+			{
+				bNeighborsAbsorbed |= AbsorbClusterNeighbors(TopNodeConnectivity, ClusterMagnet, RemainingPool);
+			}
+
+			// early termination
+			if (!bNeighborsAbsorbed)
+			{
+				break;
+			}
+		}
+
+		// Create new clusters from the cluster magnets
+		for (const FClusterMagnet& ClusterMagnet : ClusterMagnets)
+		{
+			if (ClusterMagnet.ClusteredNodes.Num() > 1)
+			{
+				TArray<int32> NewChildren = ClusterMagnet.ClusteredNodes.Array();
+				NewChildren.Sort();
+				FGeometryCollectionClusteringUtility::ClusterBonesUnderNewNode(&GeometryCollection, NewChildren[0], NewChildren, false, false);
+			}
+		}
+	}
+
+	return true;
+}
+
diff --git a/Engine/Plugins/Experimental/Fracture/Source/FractureEngine/Public/FractureEngineClustering.h b/Engine/Plugins/Experimental/Fracture/Source/FractureEngine/Public/FractureEngineClustering.h
--- a/Engine/Plugins/Experimental/Fracture/Source/FractureEngine/Public/FractureEngineClustering.h
+++ b/Engine/Plugins/Experimental/Fracture/Source/FractureEngine/Public/FractureEngineClustering.h
@@ -144,6 +144,13 @@ public:
 		FGeometryCollection& GeometryCollection,
 		TArray<int32>& InOutSelection
 	);
+
+	// Merge neighbors, and neighbors of neighbors (out to the Iterations number) to the selected clusters.
+	static bool ClusterMagnet(
+		FGeometryCollection& GeometryCollection,
+		TArray<int32>& InOutSelection,
+		int32 Iterations
+	);
 };
 
 #if UE_ENABLE_INCLUDE_ORDER_DEPRECATED_IN_5_2
diff --git a/Engine/Plugins/Experimental/GeometryCollectionPlugin/Source/GeometryCollectionNodes/Private/Dataflow/GeometryCollectionClusteringNodes.cpp b/Engine/Plugins/Experimental/GeometryCollectionPlugin/Source/GeometryCollectionNodes/Private/Dataflow/GeometryCollectionClusteringNodes.cpp
--- a/Engine/Plugins/Experimental/GeometryCollectionPlugin/Source/GeometryCollectionNodes/Private/Dataflow/GeometryCollectionClusteringNodes.cpp
+++ b/Engine/Plugins/Experimental/GeometryCollectionPlugin/Source/GeometryCollectionNodes/Private/Dataflow/GeometryCollectionClusteringNodes.cpp
@@ -46,6 +46,7 @@ namespace Dataflow
 		DATAFLOW_NODE_REGISTER_CREATION_FACTORY(FClusterMergeToNeighborsDataflowNode);
 		DATAFLOW_NODE_REGISTER_CREATION_FACTORY(FClusterMergeDataflowNode);
 		DATAFLOW_NODE_REGISTER_CREATION_FACTORY(FClusterIsolatedRootsDataflowNode);
+		DATAFLOW_NODE_REGISTER_CREATION_FACTORY(FClusterMagnetDataflowNode);
 
 		// GeometryCollection|Cluster
 		DATAFLOW_NODE_REGISTER_CREATION_FACTORY_NODE_COLORS_BY_CATEGORY("GeometryCollection|Cluster", FLinearColor(.25f, 0.45f, 0.8f), CDefaultNodeBodyTintColor);
@@ -247,4 +248,24 @@ void FClusterIsolatedRootsDataflowNode::Evaluate(Dataflow::FContext& Context, co
 	}
 }
 
+void FClusterMagnetDataflowNode::Evaluate(Dataflow::FContext& Context, const FDataflowOutput* Out) const
+{
+	if (Out->IsA(&Collection))
+	{
+		const FManagedArrayCollection& InCollection = GetValue(Context, &Collection);
+		if (TUniquePtr<FGeometryCollection> GeomCollection = TUniquePtr<FGeometryCollection>(InCollection.NewCopy<FGeometryCollection>()))
+		{
+			FDataflowTransformSelection InTransformSelection = GetValue(Context, &TransformSelection);
+			int32 InIterations = FMath::Max(1, GetValue(Context, &Iterations));
+			TArray<int32> InSelection = InTransformSelection.AsArray();
+			if (FFractureEngineClustering::ClusterMagnet(*GeomCollection, InSelection, InIterations))
+			{
+				SetValue(Context, (const FManagedArrayCollection&)(*GeomCollection), &Collection);
+				return;
+			}
+		}
+		SetValue(Context, InCollection, &Collection);
+	}
+}
+
 
diff --git a/Engine/Plugins/Experimental/GeometryCollectionPlugin/Source/GeometryCollectionNodes/Public/Dataflow/GeometryCollectionClusteringNodes.h b/Engine/Plugins/Experimental/GeometryCollectionPlugin/Source/GeometryCollectionNodes/Public/Dataflow/GeometryCollectionClusteringNodes.h
--- a/Engine/Plugins/Experimental/GeometryCollectionPlugin/Source/GeometryCollectionNodes/Public/Dataflow/GeometryCollectionClusteringNodes.h
+++ b/Engine/Plugins/Experimental/GeometryCollectionPlugin/Source/GeometryCollectionNodes/Public/Dataflow/GeometryCollectionClusteringNodes.h
@@ -323,6 +323,44 @@ public:
 
 };
 
+
+/**
+ * Cluster by grouping the selected bones with their adjacent, neighboring bones.
+ */
+USTRUCT(meta = (DataflowGeometryCollection))
+struct FClusterMagnetDataflowNode : public FDataflowNode
+{
+	GENERATED_USTRUCT_BODY()
+	DATAFLOW_NODE_DEFINE_INTERNAL(FClusterMagnetDataflowNode, "ClusterMagnet", "GeometryCollection|Cluster", "")
+
+public:
+
+	/** Collection on which to merge bones into a cluster */
+	UPROPERTY(meta = (DataflowInput, DataflowOutput, DataflowPassthrough = "Collection", DataflowIntrinsic))
+	FManagedArrayCollection Collection;
+
+	/** Bone selection */
+	UPROPERTY(meta = (DataflowInput, DisplayName = "TransformSelection"))
+	FDataflowTransformSelection TransformSelection;
+
+	/** How many layers of neighbors to include in the clusters -- i.e. if 1, only direct neighbors are clustered; if 2, neighbors of neighbors are included, etc. */
+	UPROPERTY(EditAnywhere, Category = "Cluster Magnet", meta = (ClampMin = "1", DataflowInput, DisplayName = "Iterations"))
+	int32 Iterations = 1;
+
+	FClusterMagnetDataflowNode(const Dataflow::FNodeParameters& InParam, FGuid InGuid = FGuid::NewGuid())
+		: FDataflowNode(InParam, InGuid)
+	{
+		RegisterInputConnection(&Collection);
+		RegisterInputConnection(&TransformSelection);
+		RegisterInputConnection(&Iterations);
+		RegisterOutputConnection(&Collection, &Collection);
+	}
+
+	virtual void Evaluate(Dataflow::FContext& Context, const FDataflowOutput* Out) const override;
+
+};
+
+
 namespace Dataflow
 {
 	void GeometryCollectionClusteringNodes();
diff --git a/Engine/Source/Runtime/Experimental/Chaos/Private/GeometryCollection/Facades/CollectionTransformSelectionFacade.cpp b/Engine/Source/Runtime/Experimental/Chaos/Private/GeometryCollection/Facades/CollectionTransformSelectionFacade.cpp
--- a/Engine/Source/Runtime/Experimental/Chaos/Private/GeometryCollection/Facades/CollectionTransformSelectionFacade.cpp
+++ b/Engine/Source/Runtime/Experimental/Chaos/Private/GeometryCollection/Facades/CollectionTransformSelectionFacade.cpp
@@ -756,6 +756,21 @@ namespace GeometryCollection::Facades
 		return OutSelection;
 	}
 
+	TMap<int32, TArray<int32>> FCollectionTransformSelectionFacade::GetClusteredSelections(const TArray<int32>& InSelection) const
+	{
+		TMap<int32, TArray<int32>> SiblingGroups;
+
+		// Bin the selection indices by parent index
+		const TManagedArray<int32>& Parents = ParentAttribute.Get();
+		for (int32 Index : InSelection)
+		{
+			TArray<int32>& SiblingIndices = SiblingGroups.FindOrAdd(Parents[Index]);
+			SiblingIndices.Add(Index);
+		}
+
+		return SiblingGroups;
+	}
+
 	TArray<int32> FCollectionTransformSelectionFacade::SelectVerticesInBox(const FBox& InBox, const FTransform& InBoxTransform, bool bAllVerticesInBox) const
 	{
 		TArray<int32> OutSelection;
diff --git a/Engine/Source/Runtime/Experimental/Chaos/Public/GeometryCollection/Facades/CollectionTransformSelectionFacade.h b/Engine/Source/Runtime/Experimental/Chaos/Public/GeometryCollection/Facades/CollectionTransformSelectionFacade.h
--- a/Engine/Source/Runtime/Experimental/Chaos/Public/GeometryCollection/Facades/CollectionTransformSelectionFacade.h
+++ b/Engine/Source/Runtime/Experimental/Chaos/Public/GeometryCollection/Facades/CollectionTransformSelectionFacade.h
@@ -125,6 +125,9 @@ namespace GeometryCollection::Facades
 		/**  */
 		CHAOS_API TArray<int32> SelectByVolume(float VolumeMin, float VolumeMax, bool bInclusive, bool bInsideRange) const;
 
+		/** Return a map from parent of bone to the selected bones with that parent */
+		CHAOS_API TMap<int32, TArray<int32>> GetClusteredSelections(const TArray<int32>& InSelection) const;
+
 		/** 
 		* Returns indices of bones where all the vertices or at least one vertex of the bone is inside InBox 
 		*/
