Removing unused code found while working on menus

#rnx

diff --git a/Engine/Plugins/Experimental/ChaosEditor/Source/FractureEditor/Private/FractureEditorModeToolkit.cpp b/Engine/Plugins/Experimental/ChaosEditor/Source/FractureEditor/Private/FractureEditorModeToolkit.cpp
--- a/Engine/Plugins/Experimental/ChaosEditor/Source/FractureEditor/Private/FractureEditorModeToolkit.cpp
+++ b/Engine/Plugins/Experimental/ChaosEditor/Source/FractureEditor/Private/FractureEditorModeToolkit.cpp
@@ -508,7 +508,6 @@ void FFractureEditorModeToolkit::RequestModeUITabs()
 		{ 
 			return;
 		}
-		TSharedRef<FWorkspaceItem> MenuGroup = MenuModeCategoryPtr.ToSharedRef();
 		HierarchyTabInfo.OnSpawnTab = FOnSpawnTab::CreateSP(SharedThis(this), &FFractureEditorModeToolkit::CreateHierarchyTab);
 		HierarchyTabInfo.TabLabel = LOCTEXT("FractureHierarchy", "Fracture Hierarchy");
 		HierarchyTabInfo.TabTooltip = LOCTEXT("ModesToolboxTabTooltipText", "Open the  Modes tab, which contains the active editor mode's settings.");
diff --git a/Engine/Source/Editor/Kismet/Private/BlueprintDebugger.cpp b/Engine/Source/Editor/Kismet/Private/BlueprintDebugger.cpp
--- a/Engine/Source/Editor/Kismet/Private/BlueprintDebugger.cpp
+++ b/Engine/Source/Editor/Kismet/Private/BlueprintDebugger.cpp
@@ -112,8 +112,6 @@ FBlueprintDebuggerImpl::~FBlueprintDebuggerImpl()
 		FGlobalTabmanager::Get()->UnregisterNomadTabSpawner(FBlueprintEditorTabs::BlueprintDebuggerID);
 	}
 
-	const IWorkspaceMenuStructure& MenuStructure = WorkspaceMenu::GetMenuStructure();
-
 	if (DebuggingToolsTabManager.IsValid())
 	{
 		FGlobalTabmanager::Get()->UnregisterTabSpawner(FBlueprintEditorTabs::BlueprintDebuggerID);
diff --git a/Engine/Source/Editor/LevelEditor/Private/SLevelEditor.cpp b/Engine/Source/Editor/LevelEditor/Private/SLevelEditor.cpp
--- a/Engine/Source/Editor/LevelEditor/Private/SLevelEditor.cpp
+++ b/Engine/Source/Editor/LevelEditor/Private/SLevelEditor.cpp
@@ -1751,8 +1751,6 @@ void SLevelEditor::RefreshEditorModeCommands()
 	{
 		FLevelEditorModesCommands::Register();
 	}
-	const IWorkspaceMenuStructure& MenuStructure = WorkspaceMenu::GetMenuStructure();
-	TSharedPtr<FTabManager> LevelEditorTabManager = LevelEditorModule.GetLevelEditorTabManager();
 
 	// We need to remap all the actions to commands.
 	const FLevelEditorModesCommands& Commands = FLevelEditorModesCommands::Get();
diff --git a/Engine/Source/Editor/UnrealEd/Private/Toolkits/BaseToolkit.cpp b/Engine/Source/Editor/UnrealEd/Private/Toolkits/BaseToolkit.cpp
--- a/Engine/Source/Editor/UnrealEd/Private/Toolkits/BaseToolkit.cpp
+++ b/Engine/Source/Editor/UnrealEd/Private/Toolkits/BaseToolkit.cpp
@@ -697,7 +697,6 @@ void FModeToolkit::RequestModeUITabs()
 	if (ModeUILayer.IsValid())
 	{
 		TSharedPtr<FAssetEditorModeUILayer> ModeUILayerPtr = ModeUILayer.Pin();
-		TSharedRef<FWorkspaceItem> MenuGroup = ModeUILayerPtr->GetModeMenuCategory().ToSharedRef();
 		PrimaryTabInfo.OnSpawnTab = FOnSpawnTab::CreateSP(SharedThis(this), &FModeToolkit::CreatePrimaryModePanel);
 		PrimaryTabInfo.TabLabel = LOCTEXT("ModesToolboxTab", "Mode Toolbox");
 		PrimaryTabInfo.TabTooltip = LOCTEXT("ModesToolboxTabTooltipText", "Open the  Modes tab, which contains the active editor mode's settings.");
