stop skipping the last material in the source material list on generating a new geometry collection from a geometry collection source, as the 'selection material' this was intending to skip is no longer added to geometry collections

diff --git a/Engine/Plugins/Experimental/ChaosEditor/Source/FractureEditor/Private/FractureToolGenerators.cpp b/Engine/Plugins/Experimental/ChaosEditor/Source/FractureEditor/Private/FractureToolGenerators.cpp
--- a/Engine/Plugins/Experimental/ChaosEditor/Source/FractureEditor/Private/FractureToolGenerators.cpp
+++ b/Engine/Plugins/Experimental/ChaosEditor/Source/FractureEditor/Private/FractureToolGenerators.cpp
@@ -723,8 +723,7 @@ AGeometryCollectionActor* UFractureToolGenerateAsset::ConvertActorsToGeometryCol
 				// Record the contributing source on the asset.
 				FSoftObjectPath SourceSoftObjectPath(RestCollection);
 
-				// We're not interested in recording the final material of the collection since it's inevitably the Selection material.
-				int32 NumMaterials = GeometryCollectionComponent->GetNumMaterials() - 1;
+				int32 NumMaterials = GeometryCollectionComponent->GetNumMaterials();
 				TArray<TObjectPtr<UMaterialInterface>> SourceMaterials;
 				SourceMaterials.SetNum(NumMaterials);
 				for (int32 MaterialIndex = 0; MaterialIndex < NumMaterials; ++MaterialIndex)
