CIS fix for non-unity build error in FractureToolGenerators.cpp

diff --git a/Engine/Plugins/Experimental/ChaosEditor/Source/FractureEditor/Private/FractureToolGenerators.cpp b/Engine/Plugins/Experimental/ChaosEditor/Source/FractureEditor/Private/FractureToolGenerators.cpp
--- a/Engine/Plugins/Experimental/ChaosEditor/Source/FractureEditor/Private/FractureToolGenerators.cpp
+++ b/Engine/Plugins/Experimental/ChaosEditor/Source/FractureEditor/Private/FractureToolGenerators.cpp
@@ -14,6 +14,8 @@
 #include "Engine/StaticMeshActor.h"
 #include "Framework/Application/SlateApplication.h"
 #include "Misc/Change.h"
+#include "Misc/MessageDialog.h"
+#include "Misc/ScopedSlowTask.h"
 
 #include "FractureModeSettings.h"
 
