Make the Fracture Mode's Reset tool into a modal tool so we can make "Reset Materials" an option on the tool, rather than trying to guess whether we should reset materials based on material counts

diff --git a/Engine/Plugins/Experimental/ChaosEditor/Source/FractureEditor/Private/FractureToolGenerators.cpp b/Engine/Plugins/Experimental/ChaosEditor/Source/FractureEditor/Private/FractureToolGenerators.cpp
--- a/Engine/Plugins/Experimental/ChaosEditor/Source/FractureEditor/Private/FractureToolGenerators.cpp
+++ b/Engine/Plugins/Experimental/ChaosEditor/Source/FractureEditor/Private/FractureToolGenerators.cpp
@@ -904,6 +904,25 @@ AActor* UFractureToolGenerateAsset::AddActor(ULevel* InLevel, UClass* Class)
 }
 
 
+UFractureToolResetAsset::UFractureToolResetAsset(const FObjectInitializer& ObjInit) : Super(ObjInit)
+{
+	ResetSettings = NewObject<UGeometryCollectionResetSettings>(GetTransientPackage(), UGeometryCollectionResetSettings::StaticClass());
+	ResetSettings->OwnerTool = this;
+}
+
+FText UFractureToolResetAsset::GetApplyText() const
+{ 
+	return FText(NSLOCTEXT("FractureToolReset", "FractureToolResetAction", "Reset")); 
+}
+
+TArray<UObject*> UFractureToolResetAsset::GetSettingsObjects() const
+{
+	TArray<UObject*> Settings;
+	Settings.Add(ResetSettings);
+	return Settings;
+}
+
+
 FText UFractureToolResetAsset::GetDisplayText() const
 {
 	return FText(NSLOCTEXT("Fracture", "FractureToolReset", "Reset"));
@@ -950,7 +969,7 @@ void UFractureToolResetAsset::Execute(TWeakPtr<FFractureEditorModeToolkit> InToo
 			TSharedPtr<FGeometryCollection, ESPMode::ThreadSafe> GeometryCollectionPtr = GeometryCollectionObject->GetGeometryCollection();
 			if (FGeometryCollection* GeometryCollection = GeometryCollectionPtr.Get())
 			{
-				constexpr bool bKeepPreviousMaterials = true; // written as a flag in case we want to make this optional later
+				bool bKeepPreviousMaterials = !ResetSettings->bResetMaterials;
 				TArray<TObjectPtr<UMaterialInterface>> OldMaterials;
 				if (bKeepPreviousMaterials)
 				{
@@ -982,10 +1001,10 @@ void UFractureToolResetAsset::Execute(TWeakPtr<FFractureEditorModeToolkit> InToo
 				constexpr bool bHasInternalMaterials = false;
 				GeometryCollectionObject->InitializeMaterials(bHasInternalMaterials);
 
-				// attempt to keep previously-set materials (as long as the source doesn't have even more materials)
-				int32 NewMatNum = GeometryCollectionObject->Materials.Num(), OldMatNum = OldMaterials.Num();
-				if (bKeepPreviousMaterials && NewMatNum <= OldMatNum)
+				// attempt to keep previously-set materials if requested
+				if (bKeepPreviousMaterials)
 				{
+					int32 NewMatNum = GeometryCollectionObject->Materials.Num(), OldMatNum = OldMaterials.Num();
 					GeometryCollectionObject->Materials.SetNum(OldMatNum);
 					for (int32 MatIdx = 0; MatIdx < OldMatNum; MatIdx++)
 					{
diff --git a/Engine/Plugins/Experimental/ChaosEditor/Source/FractureEditor/Private/FractureToolGenerators.h b/Engine/Plugins/Experimental/ChaosEditor/Source/FractureEditor/Private/FractureToolGenerators.h
--- a/Engine/Plugins/Experimental/ChaosEditor/Source/FractureEditor/Private/FractureToolGenerators.h
+++ b/Engine/Plugins/Experimental/ChaosEditor/Source/FractureEditor/Private/FractureToolGenerators.h
@@ -46,13 +46,30 @@ private:
 };
 
 
+
+UCLASS(config = EditorPerProjectUserSettings)
+class UGeometryCollectionResetSettings : public UFractureToolSettings
+{
+public:
+	GENERATED_BODY()
+
+	UGeometryCollectionResetSettings(const FObjectInitializer& ObjInit)
+		: Super(ObjInit)
+	{}
+
+	/** Whether to reset the materials along with the rest of the collection, or attempt to keep the current materials across the reset. */
+	UPROPERTY(EditAnywhere, Category = Settings)
+	bool bResetMaterials = true;
+};
+
+
 UCLASS()
-class UFractureToolResetAsset : public UFractureToolGenerateAsset
+class UFractureToolResetAsset : public UFractureModalTool
 {
 public:
 	GENERATED_BODY()
 
-	UFractureToolResetAsset(const FObjectInitializer& ObjInit) : Super(ObjInit) {}
+	UFractureToolResetAsset(const FObjectInitializer& ObjInit);
 
 	// UFractureActionTool Interface
 	virtual FText GetDisplayText() const override;
@@ -61,6 +78,14 @@ public:
 	virtual void RegisterUICommand(FFractureEditorCommands* BindingContext) override;
 	virtual void Execute(TWeakPtr<FFractureEditorModeToolkit> InToolkit) override;
 	virtual bool CanExecute() const override;
+
+	// UFractureModalTool Interface
+	virtual FText GetApplyText() const override;
+	virtual TArray<UObject*> GetSettingsObjects() const override;
+
+private:
+	UPROPERTY()
+	TObjectPtr<UGeometryCollectionResetSettings> ResetSettings;
 };
 
 
