Fix issue where Static option was missing from 'Set Initial Dynamic State' menu

+ Remove Sleeping option from the same menu, which would cause an ensure and not set the state because Sleeping is not allowed as an initial dynamic state.
Also remove Sleeping option from the similar menu in the State tool

diff --git a/Engine/Plugins/Experimental/ChaosEditor/Source/FractureEditor/Private/FractureToolProperties.h b/Engine/Plugins/Experimental/ChaosEditor/Source/FractureEditor/Private/FractureToolProperties.h
--- a/Engine/Plugins/Experimental/ChaosEditor/Source/FractureEditor/Private/FractureToolProperties.h
+++ b/Engine/Plugins/Experimental/ChaosEditor/Source/FractureEditor/Private/FractureToolProperties.h
@@ -10,7 +10,8 @@ UENUM(BlueprintType)
 enum class EDynamicStateOverrideEnum : uint8
 {
 	NoOverride = 0										  UMETA(DisplayName = "No Override"),
-	Sleeping = 1 /*~Chaos::EObjectStateType::Sleeping*/    UMETA(DisplayName = "Sleeping"),
+	//~ Note: Sleeping is tagged as Hidden because it can't be set as Initial Dynamic State
+	Sleeping = 1 /*~Chaos::EObjectStateType::Sleeping */    UMETA(Hidden, DisplayName = "Sleeping"),
 	Kinematic = 2 /*~Chaos::EObjectStateType::Kinematic*/  UMETA(DisplayName = "Kinematic"),
 	Static = 3    /*~Chaos::EObjectStateType::Static*/     UMETA(DisplayName = "Static")
 };
diff --git a/Engine/Plugins/Experimental/ChaosEditor/Source/FractureEditor/Private/SGeometryCollectionOutliner.cpp b/Engine/Plugins/Experimental/ChaosEditor/Source/FractureEditor/Private/SGeometryCollectionOutliner.cpp
--- a/Engine/Plugins/Experimental/ChaosEditor/Source/FractureEditor/Private/SGeometryCollectionOutliner.cpp
+++ b/Engine/Plugins/Experimental/ChaosEditor/Source/FractureEditor/Private/SGeometryCollectionOutliner.cpp
@@ -21,6 +21,7 @@
 #include "GeometryCollection/GeometryCollectionUtility.h"
 #include "GeometryCollection/Facades/CollectionAnchoringFacade.h"
 #include "GeometryCollection/Facades/CollectionHierarchyFacade.h"
+#include "GeometryCollection/GeometryCollectionSimulationTypes.h"
 #include "PhysicsProxy/GeometryCollectionPhysicsProxy.h"
 
 #include UE_INLINE_GENERATED_CPP_BY_NAME(SGeometryCollectionOutliner)
@@ -71,19 +72,21 @@ void FGeometryCollectionTreeItem::GenerateContextMenu(UToolMenu* Menu, SGeometry
 	StateSection.AddSubMenu("FractureToolSetInitialDynamicStateMenu", NSLOCTEXT("Fracture", "FractureToolSetInitialDynamicStateMenu", "Initial Dynamic State"), FText(),
 		FNewToolMenuDelegate::CreateLambda([&Outliner](UToolMenu* Menu)
 		{
-			const FName MenuEntryNames[] =
+			constexpr int32 MenuEntryNamesCount = 3;
+			const TPair<FName, EObjectStateTypeEnum> MenuEntryNameState[MenuEntryNamesCount] =
 			{
-				"NoOverride"
-				"Sleeping",
-				"Kinematic",
-				"Static"
+				{"NoOverride",EObjectStateTypeEnum::Chaos_NONE},
+				// Note: Sleeping state intentionally skipped here, as it's not a valid initial state
+				{"Kinematic",EObjectStateTypeEnum::Chaos_Object_Kinematic},
+				{"Static",EObjectStateTypeEnum::Chaos_Object_Static}
 			};
-			const int32 MenuEntryNamesCount = sizeof(MenuEntryNames) / sizeof(FName);
 
 			FToolMenuSection& StateSection = Menu->AddSection("State");
 			for (int32 Index = 0; Index < MenuEntryNamesCount; ++Index)
 			{
-				StateSection.AddMenuEntry(MenuEntryNames[Index], GetTextFromInitialDynamicState(Index), FText(), FSlateIcon(), FUIAction(FExecuteAction::CreateRaw(&Outliner, &SGeometryCollectionOutliner::SetInitialDynamicState, Index)));
+				int32 State = (int32)MenuEntryNameState[Index].Value;
+				StateSection.AddMenuEntry(MenuEntryNameState[Index].Key, GetTextFromInitialDynamicState(State), FText(), FSlateIcon(),
+					FUIAction(FExecuteAction::CreateRaw(&Outliner, &SGeometryCollectionOutliner::SetInitialDynamicState, State)));
 			}
 		}));
 
