Change default naming of a geometry collection asset to match Unreal naming convention by using GC_ prefix instead of _GeometryCollection postfix

diff --git a/Engine/Plugins/Experimental/ChaosEditor/Source/FractureEditor/Private/FractureToolGenerators.cpp b/Engine/Plugins/Experimental/ChaosEditor/Source/FractureEditor/Private/FractureToolGenerators.cpp
--- a/Engine/Plugins/Experimental/ChaosEditor/Source/FractureEditor/Private/FractureToolGenerators.cpp
+++ b/Engine/Plugins/Experimental/ChaosEditor/Source/FractureEditor/Private/FractureToolGenerators.cpp
@@ -80,6 +80,7 @@ class SCreateGeometryCollectionFromObject : public SCompoundWidget
 public:
 	SLATE_BEGIN_ARGS(SCreateGeometryCollectionFromObject)
 		: _AssetFilenameSuffix()
+		, _AssetFilenamePrefix()
 		, _HeadingText()
 		, _CreateButtonText()
 	{}
@@ -87,6 +88,9 @@ public:
 	/** The default suffix to use for the asset filename */
 	SLATE_ARGUMENT(FString, AssetFilenameSuffix)
 
+	/** The default suffix to use for the asset filename */
+	SLATE_ARGUMENT(FString, AssetFilenamePrefix)
+
 	/** The text to display at the top of the dialog */
 	SLATE_ARGUMENT(FText, HeadingText)
 
@@ -145,6 +149,9 @@ private:
 	/** The default suffix to use for the asset filename */
 	FString AssetFilenameSuffix;
 
+	/** The default prefix to use for the asset filename */
+	FString AssetFilenamePrefix;
+
 	/** The text to display as a heading for the dialog */
 	FText HeadingText;
 
@@ -175,6 +182,7 @@ private:
 void SCreateGeometryCollectionFromObject::Construct(const FArguments& InArgs, TSharedPtr<SWindow> InParentWindow)
 {
 	AssetFilenameSuffix = InArgs._AssetFilenameSuffix;
+	AssetFilenamePrefix = InArgs._AssetFilenamePrefix;
 	HeadingText = InArgs._HeadingText;
 	CreateButtonText = InArgs._CreateButtonText;
 	OnCreateAssetAction = InArgs._OnCreateAssetAction;
@@ -213,7 +221,6 @@ void SCreateGeometryCollectionFromObject::Construct(const FArguments& InArgs, TS
 			if (Actor)
 			{
 				ActorInstanceLabel += Actor->GetActorLabel();
-				ActorInstanceLabel += TEXT("_");
 				break;
 			}
 		}
@@ -223,7 +230,16 @@ void SCreateGeometryCollectionFromObject::Construct(const FArguments& InArgs, TS
 		ActorInstanceLabel = InArgs._DefaultNameOverride.ToString();
 	}
 
-	ActorInstanceLabel = UPackageTools::SanitizePackageName(ActorInstanceLabel + AssetFilenameSuffix);
+	if (!AssetFilenamePrefix.IsEmpty())
+	{
+		ActorInstanceLabel = AssetFilenamePrefix + TEXT("_") + ActorInstanceLabel;
+	}
+	if (!AssetFilenameSuffix.IsEmpty())
+	{
+		ActorInstanceLabel = ActorInstanceLabel + TEXT("_") + AssetFilenameSuffix;
+	}
+
+	ActorInstanceLabel = UPackageTools::SanitizePackageName(ActorInstanceLabel);
 
 	FString AssetName = ActorInstanceLabel;
 	FString BasePath = AssetPath / AssetName;
@@ -575,7 +591,7 @@ void UFractureToolGenerateAsset::OpenGenerateAssetDialog(TArray<AActor*>& Actors
 	TSharedPtr<SCreateGeometryCollectionFromObject> CreateAssetDialog;
 	PickAssetPathWindow->SetContent(
 		SAssignNew(CreateAssetDialog, SCreateGeometryCollectionFromObject, PickAssetPathWindow)
-		.AssetFilenameSuffix(TEXT("GeometryCollection"))
+		.AssetFilenamePrefix(TEXT("GC"))
 		.HeadingText(LOCTEXT("CreateGeometryCollection_Heading", "Geometry Collection Name"))
 		.CreateButtonText(LOCTEXT("CreateGeometryCollection_ButtonLabel", "Create Geometry Collection"))
 		.AssetPath(AssetPath)
