Add geometry column mode to fracture hierarchy view, which shows triangle and vertex counts, including cumulative counts for cluster nodes

diff --git a/Engine/Plugins/Experimental/ChaosEditor/Source/FractureEditor/Private/SGeometryCollectionOutliner.cpp b/Engine/Plugins/Experimental/ChaosEditor/Source/FractureEditor/Private/SGeometryCollectionOutliner.cpp
--- a/Engine/Plugins/Experimental/ChaosEditor/Source/FractureEditor/Private/SGeometryCollectionOutliner.cpp
+++ b/Engine/Plugins/Experimental/ChaosEditor/Source/FractureEditor/Private/SGeometryCollectionOutliner.cpp
@@ -241,6 +241,15 @@ TSharedRef<SWidget> SGeometryCollectionOutlinerRow::GenerateWidgetForColumn(cons
 		return Item->MakeImportedCollisionsColumnWidget();
 	if (ColumnName == SGeometryCollectionOutlinerColumnID::ConvexCount)
 		return Item->MakeConvexCountColumnWidget();
+	if (ColumnName == SGeometryCollectionOutlinerColumnID::TriangleCount)
+	{
+		return Item->MakeTriangleCountColumnWidget();
+	}
+	if (ColumnName == SGeometryCollectionOutlinerColumnID::VertexCount)
+	{
+		return Item->MakeVertexCountColumnWidget();
+	}
+
 	return Item->MakeEmptyColumnWidget();
 }
 
@@ -384,6 +393,24 @@ void SGeometryCollectionOutliner::RegenerateHeader()
 			.FillWidth(CustomFillWidth)
 		);
 		break;
+
+	case EOutlinerColumnMode::Geometry:
+		HeaderRowWidget->AddColumn(
+			SHeaderRow::Column(SGeometryCollectionOutlinerColumnID::TriangleCount)
+			.DefaultLabel(LOCTEXT("GCOutliner_Column_TriangleCount", "Triangle Count"))
+			.DefaultTooltip(LOCTEXT("GCOutliner_Column_TriangleCount_ToolTip", "Number of Triangles"))
+			.HAlignHeader(EHorizontalAlignment::HAlign_Center)
+			.FillWidth(CustomFillWidth)
+		);
+		HeaderRowWidget->AddColumn(
+			SHeaderRow::Column(SGeometryCollectionOutlinerColumnID::VertexCount)
+			.DefaultLabel(LOCTEXT("GCOutliner_Column_VertexCount", "Vertex Count"))
+			.DefaultTooltip(LOCTEXT("GCOutliner_Column_VertexCount_ToolTip", "Number of Vertices"))
+			.HAlignHeader(EHorizontalAlignment::HAlign_Center)
+			.FillWidth(CustomFillWidth)
+		);
+		break;
+
 	}
 }
 
@@ -648,6 +675,8 @@ FGeometryCollectionItemDataFacade::FGeometryCollectionItemDataFacade(FManagedArr
 	, HasSourceCollisionAttribute(InCollection, "HasSourceCollision", DataCollectionGroup)
 	, SourceCollisionUsedAttribute(InCollection, "SourceCollisionUsed", DataCollectionGroup)
 	, ConvexCountAttribute(InCollection, "ConvexCount", DataCollectionGroup)
+	, TriangleCountAttribute(InCollection, "TriangleCount", DataCollectionGroup)
+	, VertexCountAttribute(InCollection, "VertexCount", DataCollectionGroup)
 {
 }
 
@@ -714,6 +743,16 @@ int32 FGeometryCollectionItemDataFacade::GetConvexCount(int32 Index) const
 	return GetAttributeValue(ConvexCountAttribute, Index, 0);
 }
 
+int32 FGeometryCollectionItemDataFacade::GetTriangleCount(int32 Index) const
+{
+	return GetAttributeValue(TriangleCountAttribute, Index, 0);
+}
+
+int32 FGeometryCollectionItemDataFacade::GetVertexCount(int32 Index) const
+{
+	return GetAttributeValue(VertexCountAttribute, Index, 0);
+}
+
 void FGeometryCollectionItemDataFacade::FillFromGeometryCollectionComponent(const UGeometryCollectionComponent& GeometryCollectionComponent, EOutlinerColumnMode ColumnMode)
 {
 	const UGeometryCollection* RestCollection = GeometryCollectionComponent.GetRestCollection();
@@ -862,6 +901,37 @@ void FGeometryCollectionItemDataFacade::FillFromGeometryCollectionComponent(cons
 				RemoveOnBreakAttribute.Copy(GCRemoveOnBreakAttribute);
 			}
 		}
+		else if (ColumnMode == EOutlinerColumnMode::Geometry)
+		{
+			const TManagedArrayAccessor<int32> GCTransformToGeometryIndexAttribute(GeometryCollection, "TransformToGeometryIndex", FGeometryCollection::TransformGroup);
+			const TManagedArrayAccessor<int32> GCGeoVertexCountAttribute(GeometryCollection, "VertexCount", FGeometryCollection::GeometryGroup);
+			const TManagedArrayAccessor<int32> GCGeoFaceCountAttribute(GeometryCollection, "FaceCount", FGeometryCollection::GeometryGroup);
+			if (GCTransformToGeometryIndexAttribute.IsValid() && GCGeoVertexCountAttribute.IsValid() && GCGeoFaceCountAttribute.IsValid() && GCSimulationTypeAttribute.IsValid())
+			{
+				TManagedArray<int32>& VertexCountArray = VertexCountAttribute.Add();
+				TManagedArray<int32>& TriangleCountArray = TriangleCountAttribute.Add();
+
+				Chaos::Facades::FCollectionHierarchyFacade HierarchyFacade(GeometryCollection);
+				const TArray<int32> TransformIndices = HierarchyFacade.GetTransformArrayInDepthFirstOrder();
+
+				for (int32 TransformIndex : TransformIndices)
+				{
+					int32 GeometryIndex = GCTransformToGeometryIndexAttribute[TransformIndex];
+					if (GCSimulationTypeAttribute[TransformIndex] == FGeometryCollection::ESimulationTypes::FST_Rigid)
+					{
+						VertexCountArray[TransformIndex] = GCGeoVertexCountAttribute[GeometryIndex];
+						TriangleCountArray[TransformIndex] = GCGeoFaceCountAttribute[GeometryIndex];
+					}
+
+					int32 ParentTransformIndex = GeometryCollection.Parent[TransformIndex];
+					if (ParentTransformIndex != INDEX_NONE)
+					{
+						VertexCountArray[ParentTransformIndex] -= FMath::Abs(VertexCountArray[TransformIndex]);
+						TriangleCountArray[ParentTransformIndex] -= FMath::Abs(TriangleCountArray[TransformIndex]);
+					}
+				}
+			}
+			}
 	}
 }
 
@@ -1358,6 +1428,43 @@ TSharedRef<SWidget> FGeometryCollectionTreeItemBone::MakeConvexCountColumnWidget
 		];
 }
 
+TSharedRef<SWidget> FGeometryCollectionTreeItemBone::MakeTriangleCountColumnWidget() const
+{
+	const FGeometryCollectionItemDataFacade& DataCollectionFacade = GetDataCollectionFacade();
+	const int32 TriangleCount = DataCollectionFacade.GetTriangleCount(BoneIndex);
+	const bool bIsSumOfChildren = (TriangleCount < 0);
+
+	const FText ItemText = bIsSumOfChildren ? FText::Format(LOCTEXT("GCOutliner_TriangleCount_Format", "({0})"), FMath::Abs(TriangleCount)) : FText::AsNumber(TriangleCount);
+
+	return SNew(SHorizontalBox)
+		+ SHorizontalBox::Slot()
+		.Padding(12.f, 0.f)
+		.HAlign(HAlign_Right)
+		[
+			SNew(STextBlock)
+			.Text(ItemText)
+		.ColorAndOpacity(ItemColor)
+		];
+}
+TSharedRef<SWidget> FGeometryCollectionTreeItemBone::MakeVertexCountColumnWidget() const
+{
+	const FGeometryCollectionItemDataFacade& DataCollectionFacade = GetDataCollectionFacade();
+	const int32 VertexCount = DataCollectionFacade.GetVertexCount(BoneIndex);
+	const bool bIsSumOfChildren = (VertexCount < 0);
+
+	const FText ItemText = bIsSumOfChildren ? FText::Format(LOCTEXT("GCOutliner_VertexCount_Format", "({0})"), FMath::Abs(VertexCount)) : FText::AsNumber(VertexCount);
+
+	return SNew(SHorizontalBox)
+		+ SHorizontalBox::Slot()
+		.Padding(12.f, 0.f)
+		.HAlign(HAlign_Right)
+		[
+			SNew(STextBlock)
+			.Text(ItemText)
+		.ColorAndOpacity(ItemColor)
+		];
+}
+
 TSharedRef<SWidget> FGeometryCollectionTreeItemBone::MakeInitialStateColumnWidget() const
 {
 	const FGeometryCollectionItemDataFacade& DataCollectionFacade = GetDataCollectionFacade();
diff --git a/Engine/Plugins/Experimental/ChaosEditor/Source/FractureEditor/Private/SGeometryCollectionOutliner.h b/Engine/Plugins/Experimental/ChaosEditor/Source/FractureEditor/Private/SGeometryCollectionOutliner.h
--- a/Engine/Plugins/Experimental/ChaosEditor/Source/FractureEditor/Private/SGeometryCollectionOutliner.h
+++ b/Engine/Plugins/Experimental/ChaosEditor/Source/FractureEditor/Private/SGeometryCollectionOutliner.h
@@ -111,6 +111,8 @@ public:
 	bool HasSourceCollision(int32 Index) const;
 	bool IsSourceCollisionUsed(int32 Index) const;
 	int32 GetConvexCount(int32 Index) const;
+	int32 GetTriangleCount(int32 Index) const;
+	int32 GetVertexCount(int32 Index) const;
 
 private:
 	FManagedArrayCollection&			DataCollection;
@@ -130,6 +132,8 @@ private:
 	TManagedArrayAccessor<bool>			HasSourceCollisionAttribute;
 	TManagedArrayAccessor<bool>			SourceCollisionUsedAttribute;
 	TManagedArrayAccessor<int32>		ConvexCountAttribute;
+	TManagedArrayAccessor<int32>		TriangleCountAttribute;
+	TManagedArrayAccessor<int32>		VertexCountAttribute;
 };
 
 class FGeometryCollectionTreeItemComponent : public FGeometryCollectionTreeItem
@@ -209,6 +213,8 @@ public:
 	TSharedRef<SWidget> MakeRemovalTimeColumnWidget() const;
 	TSharedRef<SWidget> MakeImportedCollisionsColumnWidget() const;
 	TSharedRef<SWidget> MakeConvexCountColumnWidget() const;
+	TSharedRef<SWidget> MakeTriangleCountColumnWidget() const;
+	TSharedRef<SWidget> MakeVertexCountColumnWidget() const;
 	TSharedRef<SWidget> MakeEmptyColumnWidget() const;
 	virtual void GetChildren(FGeometryCollectionTreeItemList& OutChildren) override;
 	bool IsValidBone() const;
@@ -253,6 +259,9 @@ namespace SGeometryCollectionOutlinerColumnID
 	// Collision Column Mode
 	const FName ConvexCount("Convex Count");
 	const FName ImportedCollisions("ImportedCollisions");
+	// Geometry Column Mode
+	const FName VertexCount("Vertex Count");
+	const FName TriangleCount("Triangle Count");
 }
 
 class SGeometryCollectionOutlinerRow : public SMultiColumnTableRow<FGeometryCollectionTreeItemBonePtr>
diff --git a/Engine/Plugins/Experimental/ChaosEditor/Source/FractureEditor/Public/FractureEditorModeToolkit.h b/Engine/Plugins/Experimental/ChaosEditor/Source/FractureEditor/Public/FractureEditorModeToolkit.h
--- a/Engine/Plugins/Experimental/ChaosEditor/Source/FractureEditor/Public/FractureEditorModeToolkit.h
+++ b/Engine/Plugins/Experimental/ChaosEditor/Source/FractureEditor/Public/FractureEditorModeToolkit.h
@@ -95,6 +95,7 @@ enum class EOutlinerColumnMode : uint8
 	Removal = 2				UMETA(DisplayName = "Removal"),
 	Collision = 3			UMETA(DisplayName = "Collision"),
 	Size = 4				UMETA(DisplayName = "Size"),
+	Geometry = 5			UMETA(DisplayName = "Geometry")
 };
 
 class FRACTUREEDITOR_API FFractureEditorModeToolkit : public FModeToolkit, public FGCObject
