make fracture editor mode's existing selection array a uproperty() so it is visible to the editor reflection system, and also make it use weak pointers

diff --git a/Engine/Plugins/Experimental/ChaosEditor/Source/FractureEditor/Private/FractureEditorMode.cpp b/Engine/Plugins/Experimental/ChaosEditor/Source/FractureEditor/Private/FractureEditorMode.cpp
--- a/Engine/Plugins/Experimental/ChaosEditor/Source/FractureEditor/Private/FractureEditorMode.cpp
+++ b/Engine/Plugins/Experimental/ChaosEditor/Source/FractureEditor/Private/FractureEditorMode.cpp
@@ -510,11 +510,14 @@ bool UFractureEditorMode::GetPivotForOrbit(FVector& OutPivot) const
 void UFractureEditorMode::OnUndoRedo()
 {
 	RefreshOutlinerWithCurrentSelection(); // always refresh the outliner in case the geometry collection bones have changed
-	for (UGeometryCollectionComponent* SelectedComp : SelectedGeometryComponents)
+	for (TWeakObjectPtr<UGeometryCollectionComponent> SelectedCompWeakPtr : SelectedGeometryComponents)
 	{
-		// We need to update the bone colors to account for undoing/redoing selection
-		bool bForce = true;
-		FScopedColorEdit Edit(SelectedComp, bForce);
+		if (UGeometryCollectionComponent* SelectedComp = SelectedCompWeakPtr.Get())
+		{
+			// We need to update the bone colors to account for undoing/redoing selection
+			bool bForce = true;
+			FScopedColorEdit Edit(SelectedComp, bForce);
+		}
 	}
 }
 
@@ -566,8 +569,9 @@ void UFractureEditorMode::OnActorSelectionChanged(const TArray<UObject*>& NewSel
 	}
 
 	// reset state for components no longer selected
-	for (UGeometryCollectionComponent* ExistingSelection : SelectedGeometryComponents)
+	for (TWeakObjectPtr<UGeometryCollectionComponent> ExistingSelectionWeakPtr : SelectedGeometryComponents)
 	{
+		UGeometryCollectionComponent* ExistingSelection = ExistingSelectionWeakPtr.Get();
 		if (ExistingSelection && ExistingSelection->IsRegistered() && !ExistingSelection->IsBeingDestroyed() && !NewGeomSelection.Contains(ExistingSelection))
 		{
 			// This component is no longer selected, clear any modified state
@@ -591,7 +595,11 @@ void UFractureEditorMode::OnActorSelectionChanged(const TArray<UObject*>& NewSel
 		}
 	}
 
-	SelectedGeometryComponents = NewGeomSelection.Array();
+	SelectedGeometryComponents.Reset(NewGeomSelection.Num());
+	for (UGeometryCollectionComponent* Sel : NewGeomSelection)
+	{
+		SelectedGeometryComponents.Emplace(Sel);
+	}
 
 	RefreshOutlinerWithCurrentSelection();
 }
@@ -601,10 +609,22 @@ void UFractureEditorMode::RefreshOutlinerWithCurrentSelection()
 	if (Toolkit.IsValid())
 	{
 		FFractureEditorModeToolkit* FractureToolkit = (FFractureEditorModeToolkit*)Toolkit.Get();
-		FractureToolkit->SetOutlinerComponents(SelectedGeometryComponents);
+		FractureToolkit->SetOutlinerComponents(GetValidSelectedGeometryComponents());
+	}
+}
+
+TArray<UGeometryCollectionComponent*> UFractureEditorMode::GetValidSelectedGeometryComponents()
+{
+	TArray<UGeometryCollectionComponent*> Components;
+	Components.Reserve(SelectedGeometryComponents.Num());
+	for (TWeakObjectPtr<UGeometryCollectionComponent> ComponentWeakPtr : SelectedGeometryComponents)
+	{
+		Components.Add(ComponentWeakPtr.Get());
 	}
+	return Components;
 }
 
+
 void UFractureEditorMode::GetComponentGlobalBounds(UGeometryCollectionComponent* GeometryCollectionComponent, TArray<FBox>& BoundsPerBone) const
 {
 	FGeometryCollectionEdit RestCollection = GeometryCollectionComponent->EditRestCollection(GeometryCollection::EEditUpdate::None);
@@ -655,9 +675,12 @@ void UFractureEditorMode::HandlePackageReloaded(const EPackageReloadPhase InPack
 	{
 		// assemble referenced RestCollections
 		TMap<const UGeometryCollection*, UGeometryCollectionComponent*> ReferencedRestCollections;
-		for (UGeometryCollectionComponent* ExistingSelection : SelectedGeometryComponents)
+		for (TWeakObjectPtr<UGeometryCollectionComponent> ExistingSelectionWeakPtr : SelectedGeometryComponents)
 		{
-			ReferencedRestCollections.Add(TPair<const UGeometryCollection*, UGeometryCollectionComponent*>(ExistingSelection->GetRestCollection(), ExistingSelection));
+			if (UGeometryCollectionComponent* ExistingSelection = ExistingSelectionWeakPtr.Get())
+			{
+				ReferencedRestCollections.Add(TPair<const UGeometryCollection*, UGeometryCollectionComponent*>(ExistingSelection->GetRestCollection(), ExistingSelection));
+			}
 		}
 
 		// refresh outliner if reloaded package contains a referenced RestCollection
@@ -671,7 +694,7 @@ void UFractureEditorMode::HandlePackageReloaded(const EPackageReloadPhase InPack
 					{
 						FFractureEditorModeToolkit* FractureToolkit = (FFractureEditorModeToolkit*)Toolkit.Get();
 						FFractureSelectionTools::ClearSelectedBones(ReferencedRestCollections[NewObject]);
-						FractureToolkit->SetOutlinerComponents(SelectedGeometryComponents);
+						FractureToolkit->SetOutlinerComponents(GetValidSelectedGeometryComponents());
 					}
 				}
 			}
diff --git a/Engine/Plugins/Experimental/ChaosEditor/Source/FractureEditor/Public/FractureEditorMode.h b/Engine/Plugins/Experimental/ChaosEditor/Source/FractureEditor/Public/FractureEditorMode.h
--- a/Engine/Plugins/Experimental/ChaosEditor/Source/FractureEditor/Public/FractureEditorMode.h
+++ b/Engine/Plugins/Experimental/ChaosEditor/Source/FractureEditor/Public/FractureEditorMode.h
@@ -86,7 +86,12 @@ private:
 
 private:
 	/** This selection set is updated from actor selection changed event.  We change state on components as they are selected so we have to maintain or own list **/
-	TArray<UGeometryCollectionComponent*> SelectedGeometryComponents;
+	UPROPERTY()
+	TArray<TWeakObjectPtr<UGeometryCollectionComponent>> SelectedGeometryComponents;
+
+	// Return an array of all selected geometry collection components that are still valid/non-null
+	TArray<UGeometryCollectionComponent*> GetValidSelectedGeometryComponents();
+
 	// Hack: We have to set this to work around the editor defaulting to orbit around selection and breaking our custom per-bone orbiting
 	mutable TOptional<FVector> CustomOrbitPivot;
 };
