add lightweight polling of selected geometry collections to attempt to catch when the geometry collection has been changed by something outside of the fracture mode (e.g., by dataflow)

diff --git a/Engine/Plugins/Experimental/ChaosEditor/Source/FractureEditor/Private/FractureEditorMode.cpp b/Engine/Plugins/Experimental/ChaosEditor/Source/FractureEditor/Private/FractureEditorMode.cpp
--- a/Engine/Plugins/Experimental/ChaosEditor/Source/FractureEditor/Private/FractureEditorMode.cpp
+++ b/Engine/Plugins/Experimental/ChaosEditor/Source/FractureEditor/Private/FractureEditorMode.cpp
@@ -185,6 +185,27 @@ void UFractureEditorMode::Tick(FEditorViewportClient* ViewportClient, float Delt
 
 	FFractureEditorModeToolkit* FractureToolkit = (FFractureEditorModeToolkit*)Toolkit.Get();
 
+	if (FractureToolkit->IsCachedOutlinerGeometryStale(SelectedGeometryComponents))
+	{
+		// Sanitize bone selections
+		for (TWeakObjectPtr<UGeometryCollectionComponent> SelectedCompWeakPtr : SelectedGeometryComponents)
+		{
+			if (UGeometryCollectionComponent* SelectedComp = SelectedCompWeakPtr.Get())
+			{
+				FScopedColorEdit Edit(SelectedComp);
+				Edit.Sanitize();
+			}
+		}
+		// refresh and also update the cache
+		RefreshOutlinerWithCurrentSelection();
+		// if we have an active tool, relaunch it to ensure any visualizations are reset
+		if (UFractureModalTool* ActiveTool = FractureToolkit->GetActiveTool())
+		{
+			FractureToolkit->ShutdownActiveTool();
+			FractureToolkit->SetActiveTool(ActiveTool);
+		}
+	}
+
 	if (UFractureModalTool* FractureTool = FractureToolkit->GetActiveTool())
 	{
 		FractureTool->OnTick(DeltaTime);
diff --git a/Engine/Plugins/Experimental/ChaosEditor/Source/FractureEditor/Private/FractureEditorModeToolkit.cpp b/Engine/Plugins/Experimental/ChaosEditor/Source/FractureEditor/Private/FractureEditorModeToolkit.cpp
--- a/Engine/Plugins/Experimental/ChaosEditor/Source/FractureEditor/Private/FractureEditorModeToolkit.cpp
+++ b/Engine/Plugins/Experimental/ChaosEditor/Source/FractureEditor/Private/FractureEditorModeToolkit.cpp
@@ -1643,8 +1643,55 @@ FText FFractureEditorModeToolkit::GetActiveToolMessage() const
 	return LOCTEXT("FractureNoToolMessage", "Select geometry and use “New+” to create a new Geometry Collection to begin fracturing.  Choose one of the fracture tools to break apart the selected Geometry Collection.");
 }
 
+namespace UE::PrivateOutlinerCacheHelper
+{
+	void AddGeometryStatsForComponent(const UGeometryCollectionComponent* Component, int64& BoneCount, int64& VertexCount)
+	{
+		if (Component)
+		{
+			if (const UGeometryCollection* GeometryObject = Component->GetRestCollection())
+			{
+				if (const FGeometryCollection* Collection = GeometryObject->GetGeometryCollection().Get())
+				{
+					BoneCount += Collection->Transform.Num();
+					VertexCount += Collection->Vertex.Num();
+				}
+			}
+		}
+	}
+
+	void ComputeGeometryStats(const TArray<UGeometryCollectionComponent*>& InNewComponents, int64& BoneCount, int64& VertexCount)
+	{
+		BoneCount = VertexCount = 0;
+		for (const UGeometryCollectionComponent* Component : InNewComponents)
+		{
+			AddGeometryStatsForComponent(Component, BoneCount, VertexCount);
+		}
+	}
+
+	void ComputeGeometryStats(const TArray<TWeakObjectPtr<UGeometryCollectionComponent>>& InNewComponents, int64& BoneCount, int64& VertexCount)
+	{
+		BoneCount = VertexCount = 0;
+		for (const TWeakObjectPtr<UGeometryCollectionComponent> Component : InNewComponents)
+		{
+			AddGeometryStatsForComponent(Component.Get(), BoneCount, VertexCount);
+		}
+	}
+}
+
+bool FFractureEditorModeToolkit::IsCachedOutlinerGeometryStale(const TArray<TWeakObjectPtr<UGeometryCollectionComponent>>& SelectedComponents) const
+{
+	// Note we currently use quick-to-compute high level stats, compared vs cached versions, since this is run per tick; we could change this to more thoroughly walk the outliner data potentially
+	int64 NewBoneCount, NewVertexCount;
+	UE::PrivateOutlinerCacheHelper::ComputeGeometryStats(SelectedComponents, NewBoneCount, NewVertexCount);
+	return NewBoneCount != OutlinerCachedBoneCount || NewVertexCount != OutlinerCachedVertexCount;
+}
+
 void FFractureEditorModeToolkit::SetOutlinerComponents(const TArray<UGeometryCollectionComponent*>& InNewComponents)
 {
+	// Update cached stats (bone and vertex count) of the components in the outliner
+	UE::PrivateOutlinerCacheHelper::ComputeGeometryStats(InNewComponents, OutlinerCachedBoneCount, OutlinerCachedVertexCount);
+
 	TArray<UGeometryCollectionComponent*> ComponentsToEdit;
 	ComponentsToEdit.Reserve(InNewComponents.Num());
 	for (UGeometryCollectionComponent* Component : InNewComponents)
diff --git a/Engine/Plugins/Experimental/ChaosEditor/Source/FractureEditor/Public/FractureEditorModeToolkit.h b/Engine/Plugins/Experimental/ChaosEditor/Source/FractureEditor/Public/FractureEditorModeToolkit.h
--- a/Engine/Plugins/Experimental/ChaosEditor/Source/FractureEditor/Public/FractureEditorModeToolkit.h
+++ b/Engine/Plugins/Experimental/ChaosEditor/Source/FractureEditor/Public/FractureEditorModeToolkit.h
@@ -203,6 +203,9 @@ public:
 
 	void SetOutlinerColumnMode(EOutlinerColumnMode ColumnMode);
 
+	// function to poll whether the geometry collection data cached by the outliner is stale compared to the current geometry (using quick heuristics such as bone counts)
+	bool IsCachedOutlinerGeometryStale(const TArray<TWeakObjectPtr<UGeometryCollectionComponent>>& SelectedComponents) const;
+
 protected:
 	/** FModeToolkit interface */
 	virtual void RequestModeUITabs() override;
@@ -262,4 +265,7 @@ private:
 	TArray<TSharedPtr<FString>> AssetLocationModes;
 	TSharedPtr<STextComboBox> AssetLocationMode;
 
+	// Simple cached statistics to allow use to quickly/heuristically check for stale geometry collection data
+	int64 OutlinerCachedBoneCount = 0, OutlinerCachedVertexCount = 0;
+
 };
