fix crash when a geometry collection selected by fracture mode is invalidated / becomes null

diff --git a/Engine/Plugins/Experimental/ChaosEditor/Source/FractureEditor/Private/FractureEditorMode.cpp b/Engine/Plugins/Experimental/ChaosEditor/Source/FractureEditor/Private/FractureEditorMode.cpp
--- a/Engine/Plugins/Experimental/ChaosEditor/Source/FractureEditor/Private/FractureEditorMode.cpp
+++ b/Engine/Plugins/Experimental/ChaosEditor/Source/FractureEditor/Private/FractureEditorMode.cpp
@@ -640,7 +640,10 @@ TArray<UGeometryCollectionComponent*> UFractureEditorMode::GetValidSelectedGeome
 	Components.Reserve(SelectedGeometryComponents.Num());
 	for (TWeakObjectPtr<UGeometryCollectionComponent> ComponentWeakPtr : SelectedGeometryComponents)
 	{
-		Components.Add(ComponentWeakPtr.Get());
+		if (UGeometryCollectionComponent* Component = ComponentWeakPtr.Get())
+		{
+			Components.Add(Component);
+		}
 	}
 	return Components;
 }
diff --git a/Engine/Plugins/Experimental/ChaosEditor/Source/FractureEditor/Private/FractureEditorModeToolkit.cpp b/Engine/Plugins/Experimental/ChaosEditor/Source/FractureEditor/Private/FractureEditorModeToolkit.cpp
--- a/Engine/Plugins/Experimental/ChaosEditor/Source/FractureEditor/Private/FractureEditorModeToolkit.cpp
+++ b/Engine/Plugins/Experimental/ChaosEditor/Source/FractureEditor/Private/FractureEditorModeToolkit.cpp
@@ -1696,6 +1696,10 @@ void FFractureEditorModeToolkit::SetOutlinerComponents(const TArray<UGeometryCol
 	ComponentsToEdit.Reserve(InNewComponents.Num());
 	for (UGeometryCollectionComponent* Component : InNewComponents)
 	{
+		if (!Component)
+		{
+			continue;
+		}
 		FGeometryCollectionEdit RestCollection = Component->EditRestCollection(GeometryCollection::EEditUpdate::None);
 		UGeometryCollection* FracturedGeometryCollection = RestCollection.GetRestCollection();
 
