fix crash in fracture mode outliner code: detect + mark root nodes with invalid state in outliner, and stop processing these nodes

diff --git a/Engine/Plugins/Experimental/ChaosEditor/Source/FractureEditor/Private/SGeometryCollectionOutliner.cpp b/Engine/Plugins/Experimental/ChaosEditor/Source/FractureEditor/Private/SGeometryCollectionOutliner.cpp
--- a/Engine/Plugins/Experimental/ChaosEditor/Source/FractureEditor/Private/SGeometryCollectionOutliner.cpp
+++ b/Engine/Plugins/Experimental/ChaosEditor/Source/FractureEditor/Private/SGeometryCollectionOutliner.cpp
@@ -195,7 +195,8 @@ UOutlinerSettings::UOutlinerSettings(const FObjectInitializer& ObjInit)
 
 TSharedRef<SWidget> SGeometryCollectionOutlinerRow::GenerateWidgetForColumn(const FName& ColumnName)
 {
-	if (!ensure(Item->IsValidBone()))
+	// This can happen because sometimes slate retains old items until the next tick, and keeps calling callbacks on them until then
+	if (!Item->IsValidBone())
 	{
 		return Item->MakeEmptyColumnWidget();
 	}
@@ -483,6 +484,14 @@ void SGeometryCollectionOutliner::SetComponents(const TArray<UGeometryCollection
 	TGuardValue<bool> ExternalSelectionGuard(bPerformingSelection, true);
 	TreeView->ClearSelection();
 
+	// explicitly mark the root nodes as invalid before emptying, so we know we can safely ignore them in case slate still triggers callbacks for them (they will not be deleted until the tree view refresh, on tick)
+	for (TSharedPtr<FGeometryCollectionTreeItemComponent>& RootNode : RootNodes)
+	{
+		if (RootNode)
+		{
+			RootNode->Invalidate();
+		}
+	}
 	RootNodes.Empty();
 
 	for (UGeometryCollectionComponent* Component : InNewComponents)
@@ -691,6 +700,11 @@ bool FGeometryCollectionItemDataFacade::IsValidBoneIndex(int32 BoneIndex) const
 	return BoneIndex >= 0 && BoneIndex < BoneNameAttribute.Num();
 }
 
+int32 FGeometryCollectionItemDataFacade::GetBoneCount() const
+{
+	return BoneNameAttribute.Num();
+}
+
 FString FGeometryCollectionItemDataFacade::GetBoneName(int32 Index) const
 {
 	return GetAttributeValue(BoneNameAttribute, Index, FString());
@@ -1198,10 +1212,33 @@ TSharedRef<ITableRow> FGeometryCollectionTreeItemBone::MakeTreeRowWidget(const T
 	return SNew(SGeometryCollectionOutlinerRow, InOwnerTable, SharedThis(this));
 }
 
+bool FGeometryCollectionTreeItemComponent::IsValid() const
+{
+	if (bInvalidated)
+	{
+		return false;
+	}
+	if (const UGeometryCollectionComponent* Comp = GetComponent())
+	{
+		if (const UGeometryCollection* RestCollection = Comp->GetRestCollection())
+		{
+			if (const FGeometryCollection* Collection = RestCollection->GetGeometryCollection().Get())
+			{
+				return Collection->NumElements(FGeometryCollection::TransformGroup) == DataCollectionFacade.GetBoneCount();
+			}
+		}
+	}
+	return false;
+}
+
 bool FGeometryCollectionTreeItemBone::IsValidBone() const
 {
-	const FGeometryCollectionItemDataFacade& DataCollectionFacade = GetDataCollectionFacade();
-	return DataCollectionFacade.IsValidBoneIndex(BoneIndex);
+	if (ParentComponentItem && ParentComponentItem->IsValid())
+	{
+		const FGeometryCollectionItemDataFacade& DataCollectionFacade = GetDataCollectionFacade();
+		return DataCollectionFacade.IsValidBoneIndex(BoneIndex);
+	}
+	return false;
 }
 
 TSharedRef<SWidget> FGeometryCollectionTreeItemBone::MakeBoneIndexColumnWidget() const
diff --git a/Engine/Plugins/Experimental/ChaosEditor/Source/FractureEditor/Private/SGeometryCollectionOutliner.h b/Engine/Plugins/Experimental/ChaosEditor/Source/FractureEditor/Private/SGeometryCollectionOutliner.h
--- a/Engine/Plugins/Experimental/ChaosEditor/Source/FractureEditor/Private/SGeometryCollectionOutliner.h
+++ b/Engine/Plugins/Experimental/ChaosEditor/Source/FractureEditor/Private/SGeometryCollectionOutliner.h
@@ -100,6 +100,7 @@ public:
 
 	bool IsValidBoneIndex(int32 BoneIndex) const;
 	FString GetBoneName(int32 Index) const;
+	int32 GetBoneCount() const;
 	float GetRelativeSize(int32 Index) const;
 	float GetVolumetricUnit(int32 Index) const;
 	int32 GetInitialState(int32 Index) const;
@@ -146,6 +147,8 @@ public:
 	{
 		RegenerateChildren();
 	}
+	
+	virtual ~FGeometryCollectionTreeItemComponent() {}
 
 	/** FGeometryCollectionTreeItem interface */
 	virtual TSharedRef<ITableRow> MakeTreeRowWidget(const TSharedRef<STableViewBase>& InOwnerTable, bool bIsPinned = false);
@@ -165,8 +168,17 @@ public:
 
 	void SetHistogramSelection(TArray<int32>& SelectedBones);
 
+	const FGeometryCollectionItemDataFacade& GetDataCollectionFacade() const { return DataCollectionFacade; }
 	FGeometryCollectionItemDataFacade& GetDataCollectionFacade() { return DataCollectionFacade; }
 
+	bool IsValid() const;
+
+	// Mark item as unused/invalid; helpful because slate defers destroying tree items and can still run callbacks on them until tick
+	void Invalidate()
+	{
+		bInvalidated = true;
+	}
+
 private:
 	bool FilterBoneIndex(int32 BoneIndex) const;
 
@@ -186,6 +198,9 @@ private:
 	FManagedArrayCollection DataCollection;
 	FGeometryCollectionItemDataFacade DataCollectionFacade;
 
+	// track whether the item has been explicitly invalidated
+	bool bInvalidated = false;
+
 };
 
 class FGeometryCollectionTreeItemBone : public FGeometryCollectionTreeItem
