remove unneeded uproperty for fracture editor mode's selection array of weakobjectptrs

diff --git a/Engine/Plugins/Experimental/ChaosEditor/Source/FractureEditor/Public/FractureEditorMode.h b/Engine/Plugins/Experimental/ChaosEditor/Source/FractureEditor/Public/FractureEditorMode.h
--- a/Engine/Plugins/Experimental/ChaosEditor/Source/FractureEditor/Public/FractureEditorMode.h
+++ b/Engine/Plugins/Experimental/ChaosEditor/Source/FractureEditor/Public/FractureEditorMode.h
@@ -86,7 +86,6 @@ private:
 
 private:
 	/** This selection set is updated from actor selection changed event.  We change state on components as they are selected so we have to maintain or own list **/
-	UPROPERTY()
 	TArray<TWeakObjectPtr<UGeometryCollectionComponent>> SelectedGeometryComponents;
 
 	// Return an array of all selected geometry collection components that are still valid/non-null
