fix source actors not being properly removed from level by fracture 'new' tool

diff --git a/Engine/Plugins/Experimental/ChaosEditor/Source/FractureEditor/Private/FractureToolGenerators.cpp b/Engine/Plugins/Experimental/ChaosEditor/Source/FractureEditor/Private/FractureToolGenerators.cpp
--- a/Engine/Plugins/Experimental/ChaosEditor/Source/FractureEditor/Private/FractureToolGenerators.cpp
+++ b/Engine/Plugins/Experimental/ChaosEditor/Source/FractureEditor/Private/FractureToolGenerators.cpp
@@ -52,6 +52,8 @@
 #include "AssetRegistry/AssetRegistryModule.h"
 #include "PackageTools.h"
 #include "SPrimaryButton.h"
+#include "UnrealEdGlobals.h"
+#include "Editor/UnrealEdEngine.h"
 
 #include UE_INLINE_GENERATED_CPP_BY_NAME(FractureToolGenerators)
 
@@ -669,10 +671,25 @@ void UFractureToolGenerateAsset::OnGenerateAssetPathChosen(const FString& InAsse
 
 		FScopedTransaction Transaction(LOCTEXT("RemoveSourceActors", "Remove Source Actors"));
 
+		auto DestroyActor = [](AActor* Actor)
+		{
+			if (UWorld* ActorWorld = Actor->GetWorld())
+			{
+				if (GIsEditor && GUnrealEd)
+				{
+					GUnrealEd->DeleteActors(TArray{ Actor }, ActorWorld, GUnrealEd->GetSelectedActors()->GetElementSelectionSet());
+				}
+				else
+				{
+					ActorWorld->DestroyActor(Actor);
+				}
+			}
+		};
+
 		for (AActor* Actor : Actors)
 		{
 			Actor->Modify();
-			Actor->Destroy();
+			DestroyActor(Actor);
 		}
 	}
 }
