remove creation of unused settings object in fracture editor mode toolkit

diff --git a/Engine/Plugins/Experimental/ChaosEditor/Source/FractureEditor/Private/FractureEditorModeToolkit.cpp b/Engine/Plugins/Experimental/ChaosEditor/Source/FractureEditor/Private/FractureEditorModeToolkit.cpp
--- a/Engine/Plugins/Experimental/ChaosEditor/Source/FractureEditor/Private/FractureEditorModeToolkit.cpp
+++ b/Engine/Plugins/Experimental/ChaosEditor/Source/FractureEditor/Private/FractureEditorModeToolkit.cpp
@@ -1581,9 +1581,6 @@ void FFractureEditorModeToolkit::SetActiveTool(UFractureModalTool* InActiveTool)
 
 	ActiveTool = InActiveTool;
 
-	UFractureToolSettings* ToolSettings = GetMutableDefault<UFractureToolSettings>();
-	ToolSettings->OwnerTool = ActiveTool;
-
 	TArray<UObject*> Settings;
 	Settings.Add(GetMutableDefault<UFractureSettings>());
 
