make fracture cutter common settings use a shared settings object across the tools, so settings changes persist

diff --git a/Engine/Plugins/Experimental/ChaosEditor/Source/FractureEditor/Private/FractureToolBrick.cpp b/Engine/Plugins/Experimental/ChaosEditor/Source/FractureEditor/Private/FractureToolBrick.cpp
--- a/Engine/Plugins/Experimental/ChaosEditor/Source/FractureEditor/Private/FractureToolBrick.cpp
+++ b/Engine/Plugins/Experimental/ChaosEditor/Source/FractureEditor/Private/FractureToolBrick.cpp
@@ -26,7 +26,6 @@ UFractureToolBrick::UFractureToolBrick(const FObjectInitializer& ObjInit)
 {
 	BrickSettings = NewObject<UFractureBrickSettings>(GetTransientPackage(), UFractureBrickSettings::StaticClass());
 	BrickSettings->OwnerTool = this;
-	CutterSettings->bNoisePreviewToggleEnabled = false; // TODO: support noise preview for bricks
 }
 
 FText UFractureToolBrick::GetDisplayText() const
diff --git a/Engine/Plugins/Experimental/ChaosEditor/Source/FractureEditor/Private/FractureToolBrick.h b/Engine/Plugins/Experimental/ChaosEditor/Source/FractureEditor/Private/FractureToolBrick.h
--- a/Engine/Plugins/Experimental/ChaosEditor/Source/FractureEditor/Private/FractureToolBrick.h
+++ b/Engine/Plugins/Experimental/ChaosEditor/Source/FractureEditor/Private/FractureToolBrick.h
@@ -98,6 +98,11 @@ public:
 
 	void GenerateBrickTransforms(const FBox& Bounds);
 
+	virtual void ConfigureCutterSettings() override
+	{
+		CutterSettings->bNoisePreviewToggleEnabled = false; // TODO: support noise preview for bricks
+	}
+
 protected:
 	virtual void ClearVisualizations() override
 	{
diff --git a/Engine/Plugins/Experimental/ChaosEditor/Source/FractureEditor/Private/FractureToolCutter.cpp b/Engine/Plugins/Experimental/ChaosEditor/Source/FractureEditor/Private/FractureToolCutter.cpp
--- a/Engine/Plugins/Experimental/ChaosEditor/Source/FractureEditor/Private/FractureToolCutter.cpp
+++ b/Engine/Plugins/Experimental/ChaosEditor/Source/FractureEditor/Private/FractureToolCutter.cpp
@@ -137,10 +137,8 @@ void UFractureCutterSettings::TransferNoiseSettings(FNoiseSettings& NoiseSetting
 UFractureToolCutterBase::UFractureToolCutterBase(const FObjectInitializer& ObjInit)
 	: Super(ObjInit)
 {
-	CutterSettings = NewObject<UFractureCutterSettings>(GetTransientPackage(), UFractureCutterSettings::StaticClass());
-	CutterSettings->OwnerTool = this;
-	CollisionSettings = NewObject<UFractureCollisionSettings>(GetTransientPackage(), UFractureCollisionSettings::StaticClass());
-	CollisionSettings->OwnerTool = this;
+	CutterSettings = GetMutableDefault<UFractureCutterSettings>();
+	CollisionSettings = GetMutableDefault<UFractureCollisionSettings>();
 }
 
 bool UFractureToolCutterBase::CanExecute() const
diff --git a/Engine/Plugins/Experimental/ChaosEditor/Source/FractureEditor/Private/FractureToolCutter.h b/Engine/Plugins/Experimental/ChaosEditor/Source/FractureEditor/Private/FractureToolCutter.h
--- a/Engine/Plugins/Experimental/ChaosEditor/Source/FractureEditor/Private/FractureToolCutter.h
+++ b/Engine/Plugins/Experimental/ChaosEditor/Source/FractureEditor/Private/FractureToolCutter.h
@@ -108,7 +108,7 @@ public:
 
 	/** Whether to show a solid preview of the cutting geometry, including any noise displacement */
 	UPROPERTY(EditAnywhere, Category = Visualization, meta = (EditCondition = "bNoisePreviewToggleEnabled", HideEditConditionToggle, EditConditionHides))
-	bool bDrawNoisePreview = false;
+	bool bDrawNoisePreview = true;
 
 	// This flag allows tools to disable the above bDrawNoisePreview option if/when it is not applicable
 	UPROPERTY()
@@ -234,6 +234,19 @@ public:
 	{
 		Super::Setup(InToolkit);
 		CutterSettings->UpdateActiveMaterialNames(GetSelectedComponentMaterialNames(true));
+		CutterSettings->OwnerTool = this;
+		CollisionSettings->OwnerTool = this;
+		ConfigureCutterSettings();
+	}
+
+	// Set tool-specific defaults for the cutter settings (e.g., disable grout/noise if needed)
+	virtual void ConfigureCutterSettings()
+	{
+		CutterSettings->bDrawSitesToggleEnabled = true;
+		CutterSettings->bNoisePreviewToggleEnabled = true;
+		CutterSettings->bNoisePreviewHasScale = false;
+		CutterSettings->bGroutSettingEnabled = true;
+		CutterSettings->bNoiseSettingsEnabled = true;
 	}
 
 	virtual void SelectedBonesChanged() override
@@ -261,18 +274,6 @@ public:
 		}
 	}
 
-	void DisableGroutSetting()
-	{
-		CutterSettings->bGroutSettingEnabled = false;
-		CutterSettings->Grout = 0;
-	}
-
-	void DisableNoiseSettings()
-	{
-		CutterSettings->bNoiseSettingsEnabled = false;
-		CutterSettings->Amplitude = 0;
-	}
-
 protected:
 	UPROPERTY(EditAnywhere, Category = Slicing)
 	TObjectPtr<UFractureCutterSettings> CutterSettings;
diff --git a/Engine/Plugins/Experimental/ChaosEditor/Source/FractureEditor/Private/FractureToolMeshCut.cpp b/Engine/Plugins/Experimental/ChaosEditor/Source/FractureEditor/Private/FractureToolMeshCut.cpp
--- a/Engine/Plugins/Experimental/ChaosEditor/Source/FractureEditor/Private/FractureToolMeshCut.cpp
+++ b/Engine/Plugins/Experimental/ChaosEditor/Source/FractureEditor/Private/FractureToolMeshCut.cpp
@@ -28,10 +28,6 @@ UFractureToolMeshCut::UFractureToolMeshCut(const FObjectInitializer& ObjInit)
 {
 	MeshCutSettings = NewObject<UFractureMeshCutSettings>(GetTransientPackage(), UFractureMeshCutSettings::StaticClass());
 	MeshCutSettings->OwnerTool = this;
-	CutterSettings->bDrawSitesToggleEnabled = false;
-	CutterSettings->bNoisePreviewToggleEnabled = false;
-	DisableNoiseSettings();
-	DisableGroutSetting();
 }
 
 FText UFractureToolMeshCut::GetDisplayText() const
diff --git a/Engine/Plugins/Experimental/ChaosEditor/Source/FractureEditor/Private/FractureToolMeshCut.h b/Engine/Plugins/Experimental/ChaosEditor/Source/FractureEditor/Private/FractureToolMeshCut.h
--- a/Engine/Plugins/Experimental/ChaosEditor/Source/FractureEditor/Private/FractureToolMeshCut.h
+++ b/Engine/Plugins/Experimental/ChaosEditor/Source/FractureEditor/Private/FractureToolMeshCut.h
@@ -109,6 +109,15 @@ public:
 	virtual void FractureContextChanged() override;
 	virtual int32 ExecuteFracture(const FFractureToolContext& FractureContext) override;
 
+	virtual void ConfigureCutterSettings() override
+	{
+		CutterSettings->bDrawSitesToggleEnabled = false;
+		CutterSettings->bNoisePreviewToggleEnabled = false;
+		CutterSettings->bNoisePreviewHasScale = true;
+		CutterSettings->bGroutSettingEnabled = false;
+		CutterSettings->bNoiseSettingsEnabled = false;
+	}
+
 protected:
 	virtual void ClearVisualizations() override
 	{
diff --git a/Engine/Plugins/Experimental/ChaosEditor/Source/FractureEditor/Private/FractureToolPlaneCut.cpp b/Engine/Plugins/Experimental/ChaosEditor/Source/FractureEditor/Private/FractureToolPlaneCut.cpp
--- a/Engine/Plugins/Experimental/ChaosEditor/Source/FractureEditor/Private/FractureToolPlaneCut.cpp
+++ b/Engine/Plugins/Experimental/ChaosEditor/Source/FractureEditor/Private/FractureToolPlaneCut.cpp
@@ -146,9 +146,6 @@ void UFractureToolPlaneCut::Setup(TWeakPtr<FFractureEditorModeToolkit> InToolkit
 	GizmoSettings->Setup(this);
 	PlaneCutSettings->bCanCutWithMultiplePlanes = !GizmoSettings->bUseGizmo;
 	NotifyOfPropertyChangeByTool(PlaneCutSettings);
-	CutterSettings->bDrawSitesToggleEnabled = false;
-	CutterSettings->bNoisePreviewHasScale = true;
-	CutterSettings->bDrawNoisePreview = true; // default-enable the plane noise preview
 
 	// Initialize the background compute object for the noise preview
 	if (GEditor && !NoisePreview)
diff --git a/Engine/Plugins/Experimental/ChaosEditor/Source/FractureEditor/Private/FractureToolPlaneCut.h b/Engine/Plugins/Experimental/ChaosEditor/Source/FractureEditor/Private/FractureToolPlaneCut.h
--- a/Engine/Plugins/Experimental/ChaosEditor/Source/FractureEditor/Private/FractureToolPlaneCut.h
+++ b/Engine/Plugins/Experimental/ChaosEditor/Source/FractureEditor/Private/FractureToolPlaneCut.h
@@ -63,6 +63,12 @@ public:
 	// IDynamicMeshOperatorFactory API, for generating noise preview meshes
 	virtual TUniquePtr<UE::Geometry::FDynamicMeshOperator> MakeNewOperator() override;
 
+	virtual void ConfigureCutterSettings() override
+	{
+		CutterSettings->bDrawSitesToggleEnabled = false;
+		CutterSettings->bNoisePreviewHasScale = true;
+	}
+
 protected:
 	virtual void ClearVisualizations() override;
 
diff --git a/Engine/Plugins/Experimental/ChaosEditor/Source/FractureEditor/Private/FractureToolSlice.cpp b/Engine/Plugins/Experimental/ChaosEditor/Source/FractureEditor/Private/FractureToolSlice.cpp
--- a/Engine/Plugins/Experimental/ChaosEditor/Source/FractureEditor/Private/FractureToolSlice.cpp
+++ b/Engine/Plugins/Experimental/ChaosEditor/Source/FractureEditor/Private/FractureToolSlice.cpp
@@ -23,8 +23,6 @@ UFractureToolSlice::UFractureToolSlice(const FObjectInitializer& ObjInit)
 {
 	SliceSettings = NewObject<UFractureSliceSettings>(GetTransientPackage(), UFractureSliceSettings::StaticClass());
 	SliceSettings->OwnerTool = this;
-
-	CutterSettings->bDrawSitesToggleEnabled = false;
 }
 
 FText UFractureToolSlice::GetDisplayText() const
diff --git a/Engine/Plugins/Experimental/ChaosEditor/Source/FractureEditor/Private/FractureToolSlice.h b/Engine/Plugins/Experimental/ChaosEditor/Source/FractureEditor/Private/FractureToolSlice.h
--- a/Engine/Plugins/Experimental/ChaosEditor/Source/FractureEditor/Private/FractureToolSlice.h
+++ b/Engine/Plugins/Experimental/ChaosEditor/Source/FractureEditor/Private/FractureToolSlice.h
@@ -69,6 +69,12 @@ public:
 
 	void GenerateSliceTransforms(const FFractureToolContext& Context, TArray<FTransform>& CuttingPlaneTransforms);
 
+	virtual void ConfigureCutterSettings() override
+	{
+		CutterSettings->bDrawSitesToggleEnabled = false;
+		CutterSettings->bNoisePreviewToggleEnabled = false; // TODO: support noise preview for slices
+	}
+
 	// Slicing
 	UPROPERTY(EditAnywhere, Category = Slicing)
 	TObjectPtr<UFractureSliceSettings> SliceSettings;
