add dataflow picker to new geometry collection creation ui
+ make the 'new' dialog remember its last settings across usages
+ clarify the tooltips for the settings

diff --git a/Engine/Plugins/Experimental/ChaosEditor/Source/FractureEditor/FractureEditor.Build.cs b/Engine/Plugins/Experimental/ChaosEditor/Source/FractureEditor/FractureEditor.Build.cs
--- a/Engine/Plugins/Experimental/ChaosEditor/Source/FractureEditor/FractureEditor.Build.cs
+++ b/Engine/Plugins/Experimental/ChaosEditor/Source/FractureEditor/FractureEditor.Build.cs
@@ -36,7 +36,9 @@ public class FractureEditor : ModuleRules
                 "GeometryCollectionEditor",
 				"ModelingComponents",
 				"ModelingOperators",
+				"PropertyEditor",
 				"GeometryCore",
+				"DataflowEngine",
 				"DynamicMesh",
 				"MeshDescription",
 				"StaticMeshDescription",
diff --git a/Engine/Plugins/Experimental/ChaosEditor/Source/FractureEditor/Private/FractureToolGenerators.cpp b/Engine/Plugins/Experimental/ChaosEditor/Source/FractureEditor/Private/FractureToolGenerators.cpp
--- a/Engine/Plugins/Experimental/ChaosEditor/Source/FractureEditor/Private/FractureToolGenerators.cpp
+++ b/Engine/Plugins/Experimental/ChaosEditor/Source/FractureEditor/Private/FractureToolGenerators.cpp
@@ -54,6 +54,9 @@
 #include "SPrimaryButton.h"
 #include "UnrealEdGlobals.h"
 #include "Editor/UnrealEdEngine.h"
+#include "PropertyCustomizationHelpers.h"
+
+#include "Dataflow/DataflowObject.h"
 
 #include UE_INLINE_GENERATED_CPP_BY_NAME(FractureToolGenerators)
 
@@ -75,7 +78,7 @@ namespace
 //////////////////////////////////////////////////////////////////////////
 // SCreateGeometryCollectionFromObject is adapted from SCreateAssetFromActor 
 
-DECLARE_DELEGATE_ThreeParams(FOnPathChosen, const FString&, bool, bool);
+DECLARE_DELEGATE_FourParams(FOnPathChosen, const FString&, bool, bool, UDataflow*);
 
 class SCreateGeometryCollectionFromObject : public SCompoundWidget
 {
@@ -85,6 +88,9 @@ public:
 		, _AssetFilenamePrefix()
 		, _HeadingText()
 		, _CreateButtonText()
+		, _bSplitIslands(false)
+		, _bReimportToMeshOutput(false)
+		, _DataflowAsset(nullptr)
 	{}
 
 	/** The default suffix to use for the asset filename */
@@ -106,6 +112,7 @@ public:
 
 	SLATE_ARGUMENT(bool, bSplitIslands)
 	SLATE_ARGUMENT(bool, bReimportToMeshOutput)
+	SLATE_ARGUMENT(UDataflow*, DataflowAsset)
 
 	/** Action to perform when create clicked */
 	SLATE_EVENT(FOnPathChosen, OnCreateAssetAction)
@@ -123,6 +130,22 @@ private:
 	/** Callback when the selected asset path has changed. */
 	void OnSelectAssetPath(const FString& Path);
 
+	FString GetDataflowAssetPath() const
+	{
+		if (DataflowAsset)
+		{
+			return DataflowAsset->GetPathName();
+		}
+		return FString();
+	}
+	void OnDataflowAssetChanged(const FAssetData& InAssetData)
+	{
+		if (UDataflow* NewDataflow = Cast<UDataflow>(InAssetData.GetAsset()))
+		{
+			DataflowAsset = NewDataflow;
+		}
+	}
+
 	/** Destroys the window when the operation is cancelled. */
 	FReply OnCancelCreateAssetFromActor();
 
@@ -160,12 +183,15 @@ private:
 	/** The label to be displayed on the create button */
 	FText CreateButtonText;
 
-	/** Whether to split out each connected components into separate geometry */
+	/** Whether to split out topologically-disconnected components into separate bones in the Geometry Collection */
 	bool bSplitIslands = false;
 
-	/** Whether static mesh inputs were created by the 'ToMesh' tool in Fracture Mode */
+	/** Designate input faces with odd-numbered material indices as 'internal' materials. Only use this option to preserve internal faces on meshes created with the Fracture Mode 'ToMesh' tool. */
 	bool bFromToMeshTool = false;
 
+	/** Dataflow asset to assign */
+	TObjectPtr<UDataflow> DataflowAsset = nullptr;
+
 	/** Filename textbox widget */
 	TSharedPtr<SEditableTextBox> FileNameWidget;
 
@@ -188,6 +214,10 @@ void SCreateGeometryCollectionFromObject::Construct(const FArguments& InArgs, TS
 	HeadingText = InArgs._HeadingText;
 	CreateButtonText = InArgs._CreateButtonText;
 	OnCreateAssetAction = InArgs._OnCreateAssetAction;
+	bSplitIslands = InArgs._bSplitIslands;
+	bFromToMeshTool = InArgs._bReimportToMeshOutput;
+	DataflowAsset = InArgs._DataflowAsset;
+
 
 	if (InArgs._AssetPath.IsEmpty())
 	{
@@ -339,7 +369,7 @@ void SCreateGeometryCollectionFromObject::Construct(const FArguments& InArgs, TS
 						.ToolTipText(LOCTEXT("SplitComponentsToolTip", "If checked, triangles that are not topologically connected will be assigned separate bones on import. Does not apply if creating from a Geometry Collection source."))
 					]
 
-				+ SHorizontalBox::Slot()
+					+ SHorizontalBox::Slot()
 					.FillWidth(1.0f)
 					[
 						SNew(SCheckBox)
@@ -358,6 +388,33 @@ void SCreateGeometryCollectionFromObject::Construct(const FArguments& InArgs, TS
 						.ToolTipText(LOCTEXT("SplitComponentsToolTip", "If checked, triangles that are not topologically connected will be assigned separate bones on import. Does not apply if creating from a Geometry Collection source."))
 					]
 				]
+				+ SVerticalBox::Slot()
+				[
+					SNew(SHorizontalBox)
+					+ SHorizontalBox::Slot()
+					.Padding(0.0f, 0.0f, 8.0f, 0.0f)
+					.FillWidth(0.5f)
+					.VAlign(VAlign_Center)
+					[
+						SNew(STextBlock)
+						.Text(LOCTEXT("DataflowAssetLabel", "Dataflow"))
+						.ToolTipText(LOCTEXT("DataflowAssetToolTip", "Optionally assign a Dataflow to the new Geometry Collection, to procedurally define its fracture and clustering."))
+					]
+
+					+ SHorizontalBox::Slot()
+					.FillWidth(1.0f)
+					[
+						SNew(SObjectPropertyEntryBox)
+							.ObjectPath(this, &SCreateGeometryCollectionFromObject::GetDataflowAssetPath)
+							.OnObjectChanged(this, &SCreateGeometryCollectionFromObject::OnDataflowAssetChanged)
+							.AllowedClass(UDataflow::StaticClass())
+							.DisplayUseSelected(true)
+							.DisplayBrowse(true)
+							.DisplayThumbnail(true)
+							.AllowCreate(true)
+							.AllowClear(true)
+					]
+				]
 			]
 		]
 
@@ -404,7 +461,7 @@ void SCreateGeometryCollectionFromObject::RequestDestroyParentWindow()
 FReply SCreateGeometryCollectionFromObject::OnCreateAssetFromActorClicked()
 {
 	RequestDestroyParentWindow();
-	OnCreateAssetAction.ExecuteIfBound(AssetPath / FileNameWidget->GetText().ToString(), bFromToMeshTool, bSplitIslands);
+	OnCreateAssetAction.ExecuteIfBound(AssetPath / FileNameWidget->GetText().ToString(), bFromToMeshTool, bSplitIslands, DataflowAsset);
 	return FReply::Handled();
 }
 
@@ -597,6 +654,9 @@ void UFractureToolGenerateAsset::OpenGenerateAssetDialog(TArray<AActor*>& Actors
 		.HeadingText(LOCTEXT("CreateGeometryCollection_Heading", "Geometry Collection Name"))
 		.CreateButtonText(LOCTEXT("CreateGeometryCollection_ButtonLabel", "Create Geometry Collection"))
 		.AssetPath(AssetPath)
+		.DataflowAsset(LastDataflowAsset)
+		.bReimportToMeshOutput(bLastFromToMesh)
+		.bSplitIslands(bLastSplitComponents)
 		.OnCreateAssetAction(FOnPathChosen::CreateUObject(this, &UFractureToolGenerateAsset::OnGenerateAssetPathChosen, Actors))
 	);
 
@@ -612,14 +672,17 @@ void UFractureToolGenerateAsset::OpenGenerateAssetDialog(TArray<AActor*>& Actors
 
 }
 
-void UFractureToolGenerateAsset::OnGenerateAssetPathChosen(const FString& InAssetPath, bool bFromToMeshTool, bool bSplitComponents, TArray<AActor*> Actors)
+void UFractureToolGenerateAsset::OnGenerateAssetPathChosen(const FString& InAssetPath, bool bFromToMeshTool, bool bSplitComponents, UDataflow* Dataflow, TArray<AActor*> Actors)
 {	
-	//Record the path
+	//Record the path and other settings
 	int32 LastSlash = INDEX_NONE;
 	if (InAssetPath.FindLastChar('/', LastSlash))
 	{
 		AssetPath = InAssetPath.Left(LastSlash);
 	}
+	bLastFromToMesh = bFromToMeshTool;
+	bLastSplitComponents = bSplitComponents;
+	LastDataflowAsset = Dataflow;
 
 	UGeometryCollectionComponent* GeometryCollectionComponent = nullptr;
 
@@ -629,7 +692,7 @@ void UFractureToolGenerateAsset::OnGenerateAssetPathChosen(const FString& InAsse
 
 		AGeometryCollectionActor* GeometryCollectionActor = nullptr;
 		
-		GeometryCollectionActor = ConvertActorsToGeometryCollection(InAssetPath, false/*bAddInternalMaterials*/, bSplitComponents, Actors, bFromToMeshTool);
+		GeometryCollectionActor = ConvertActorsToGeometryCollection(InAssetPath, false/*bAddInternalMaterials*/, bSplitComponents, Dataflow, Actors, bFromToMeshTool);
 		if (!GeometryCollectionActor)
 		{
 			return;
@@ -694,7 +757,7 @@ void UFractureToolGenerateAsset::OnGenerateAssetPathChosen(const FString& InAsse
 	}
 }
 
-AGeometryCollectionActor* UFractureToolGenerateAsset::ConvertActorsToGeometryCollection(const FString& InAssetPath, bool bAddInternalMaterials, bool bSplitComponents, TArray<AActor*>& Actors, bool bFromToMeshTool)
+AGeometryCollectionActor* UFractureToolGenerateAsset::ConvertActorsToGeometryCollection(const FString& InAssetPath, bool bAddInternalMaterials, bool bSplitComponents, UDataflow* Dataflow, TArray<AActor*>& Actors, bool bFromToMeshTool)
 {
 	ensure(!bAddInternalMaterials); // we should not use the 'add internal materials' path anymore, as we move away from the odd-numbered-internal-material convention
 	ensure(Actors.Num() > 0);
@@ -838,6 +901,11 @@ AGeometryCollectionActor* UFractureToolGenerateAsset::ConvertActorsToGeometryCol
 
 	FracturedGeometryCollection->InitializeMaterials(bAddInternalMaterials);
 
+	if (Dataflow)
+	{
+		FracturedGeometryCollection->DataflowAsset = Dataflow;
+	}
+
 	AddSingleRootNodeIfRequired(FracturedGeometryCollection);
 
 	FracturedGeometryCollection->InvalidateCollection();
diff --git a/Engine/Plugins/Experimental/ChaosEditor/Source/FractureEditor/Private/FractureToolGenerators.h b/Engine/Plugins/Experimental/ChaosEditor/Source/FractureEditor/Private/FractureToolGenerators.h
--- a/Engine/Plugins/Experimental/ChaosEditor/Source/FractureEditor/Private/FractureToolGenerators.h
+++ b/Engine/Plugins/Experimental/ChaosEditor/Source/FractureEditor/Private/FractureToolGenerators.h
@@ -8,6 +8,7 @@
 
 class AGeometryCollectionActor;
 class UGeometryCollection;
+class UDataflow;
 
 UCLASS()
 class UFractureToolGenerateAsset : public UFractureActionTool
@@ -30,8 +31,8 @@ public:
 
 private:
 	void OpenGenerateAssetDialog(TArray<AActor*>& Actors);
-	void OnGenerateAssetPathChosen(const FString& InAssetPath, bool bFromToMeshTool, bool bSplitComponents, TArray<AActor*> Actors);
-	AGeometryCollectionActor* ConvertActorsToGeometryCollection(const FString& InAssetPath, bool bAddInternalMaterials, bool bSplitComponents, TArray<AActor*>& Actors, bool bFromToMeshTool);
+	void OnGenerateAssetPathChosen(const FString& InAssetPath, bool bFromToMeshTool, bool bSplitComponents, UDataflow* Dataflow, TArray<AActor*> Actors);
+	AGeometryCollectionActor* ConvertActorsToGeometryCollection(const FString& InAssetPath, bool bAddInternalMaterials, bool bSplitComponents, UDataflow* Dataflow, TArray<AActor*>& Actors, bool bFromToMeshTool);
 	class AGeometryCollectionActor* CreateNewGeometryActor(const FString& Name, const FTransform& Transform, bool AddMaterials /*= false*/);
 	static ULevel* GetSelectedLevel();
 	static AActor* AddActor(ULevel* InLevel, UClass* Class);
@@ -43,6 +44,12 @@ private:
 
 	UPROPERTY()
 	FString AssetPath;
+
+	UPROPERTY()
+	TObjectPtr<UDataflow> LastDataflowAsset = nullptr;
+
+	bool bLastSplitComponents = false;
+	bool bLastFromToMesh = false;
 };
 
 
