enable view options settings to the fracture mode 'new' path picker ui, so users can see plugin content paths if desired

diff --git a/Engine/Plugins/Experimental/ChaosEditor/Source/FractureEditor/Private/FractureToolGenerators.cpp b/Engine/Plugins/Experimental/ChaosEditor/Source/FractureEditor/Private/FractureToolGenerators.cpp
--- a/Engine/Plugins/Experimental/ChaosEditor/Source/FractureEditor/Private/FractureToolGenerators.cpp
+++ b/Engine/Plugins/Experimental/ChaosEditor/Source/FractureEditor/Private/FractureToolGenerators.cpp
@@ -231,6 +231,7 @@ void SCreateGeometryCollectionFromObject::Construct(const FArguments& InArgs, TS
 	bIsReportingError = false;
 
 	FPathPickerConfig PathPickerConfig;
+	PathPickerConfig.bShowViewOptions = true;
 	PathPickerConfig.DefaultPath = AssetPath;
 	PathPickerConfig.OnPathSelected = FOnPathSelected::CreateRaw(this, &SCreateGeometryCollectionFromObject::OnSelectAssetPath);
 
