fix crash on creating a new geometry collection with actors from multiple levels + place newly-created geometry collections on the active level instead of the source actor level

diff --git a/Engine/Plugins/Experimental/ChaosEditor/Source/FractureEditor/Private/FractureToolGenerators.cpp b/Engine/Plugins/Experimental/ChaosEditor/Source/FractureEditor/Private/FractureToolGenerators.cpp
--- a/Engine/Plugins/Experimental/ChaosEditor/Source/FractureEditor/Private/FractureToolGenerators.cpp
+++ b/Engine/Plugins/Experimental/ChaosEditor/Source/FractureEditor/Private/FractureToolGenerators.cpp
@@ -933,7 +933,8 @@ class AGeometryCollectionActor* UFractureToolGenerateAsset::CreateNewGeometryAct
 	UGeometryCollection* InGeometryCollection = static_cast<UGeometryCollection*>(NewObject<UGeometryCollection>(Package, UGeometryCollection::StaticClass(), FName(*UniqueAssetName), RF_Transactional | RF_Public | RF_Standalone));
 
 	// Create the new Geometry Collection actor
-	AGeometryCollectionActor* NewActor = Cast<AGeometryCollectionActor>(AddActor(GetSelectedLevel(), AGeometryCollectionActor::StaticClass()));
+	ULevel* PlaceInLevel = GEditor->GetEditorWorldContext().World()->GetCurrentLevel();
+	AGeometryCollectionActor* NewActor = Cast<AGeometryCollectionActor>(AddActor(PlaceInLevel, AGeometryCollectionActor::StaticClass()));
 	check(NewActor->GetGeometryCollectionComponent());
 
 	// Set the Geometry Collection asset in the new actor
@@ -952,22 +953,6 @@ class AGeometryCollectionActor* UFractureToolGenerateAsset::CreateNewGeometryAct
 	return NewActor;
 }
 
-ULevel* UFractureToolGenerateAsset::GetSelectedLevel()
-{
-	USelection* SelectedActors = GEditor->GetSelectedActors();
-	TArray<ULevel*> UniqueLevels;
-	for (FSelectionIterator Iter(*SelectedActors); Iter; ++Iter)
-	{
-		AActor* Actor = Cast<AActor>(*Iter);
-		if (Actor)
-		{
-			UniqueLevels.AddUnique(Actor->GetLevel());
-		}
-	}
-	check(UniqueLevels.Num() == 1);
-	return UniqueLevels[0];
-}
-
 AActor* UFractureToolGenerateAsset::AddActor(ULevel* InLevel, UClass* Class)
 {
 	check(Class);
diff --git a/Engine/Plugins/Experimental/ChaosEditor/Source/FractureEditor/Private/FractureToolGenerators.h b/Engine/Plugins/Experimental/ChaosEditor/Source/FractureEditor/Private/FractureToolGenerators.h
--- a/Engine/Plugins/Experimental/ChaosEditor/Source/FractureEditor/Private/FractureToolGenerators.h
+++ b/Engine/Plugins/Experimental/ChaosEditor/Source/FractureEditor/Private/FractureToolGenerators.h
@@ -34,7 +34,6 @@ private:
 	void OnGenerateAssetPathChosen(const FString& InAssetPath, bool bFromToMeshTool, bool bSplitComponents, UDataflow* Dataflow, TArray<AActor*> Actors);
 	AGeometryCollectionActor* ConvertActorsToGeometryCollection(const FString& InAssetPath, bool bAddInternalMaterials, bool bSplitComponents, UDataflow* Dataflow, TArray<AActor*>& Actors, bool bFromToMeshTool);
 	class AGeometryCollectionActor* CreateNewGeometryActor(const FString& Name, const FTransform& Transform, bool AddMaterials /*= false*/);
-	static ULevel* GetSelectedLevel();
 	static AActor* AddActor(ULevel* InLevel, UClass* Class);
 
 	FString GetDefaultAssetPath(const TArray<AActor*>& Actors) const;
