prevent modeling tools from enabling element selection on non-editable dynamic meshes
+ tag the fracture editor's voronoi preview meshes as non-editable so that they won't be picked up by the element selection system (in the brief second when modeling mode is entered and fracture mode is not yet exited)

diff --git a/Engine/Plugins/Experimental/ChaosEditor/Source/FractureEditor/Private/FractureToolCutter.cpp b/Engine/Plugins/Experimental/ChaosEditor/Source/FractureEditor/Private/FractureToolCutter.cpp
--- a/Engine/Plugins/Experimental/ChaosEditor/Source/FractureEditor/Private/FractureToolCutter.cpp
+++ b/Engine/Plugins/Experimental/ChaosEditor/Source/FractureEditor/Private/FractureToolCutter.cpp
@@ -458,6 +458,7 @@ void UFractureToolVoronoiCutterBase::AddNoiseVisualizations(TArray<FCellNoisePre
 		VoronoiNoisePreviews.Add(DynamicMeshComponent);
 		DynamicMeshComponent->GetMesh()->Copy(Res.NoiseMesh);
 		DynamicMeshComponent->NotifyMeshUpdated();
+		DynamicMeshComponent->SetIsEditable(false);
 
 		// TODO: this material-access logic can go to the base class and be used by FractureToolPlaneCut as well
 		// Note the ToolSetupUtil functions optionally take a UInteractiveToolManager to provide fallback materials,
