On fracture mode selection change also update bone color and selection based on current selection level

diff --git a/Engine/Plugins/Experimental/ChaosEditor/Source/FractureEditor/Private/FractureEditorMode.cpp b/Engine/Plugins/Experimental/ChaosEditor/Source/FractureEditor/Private/FractureEditorMode.cpp
--- a/Engine/Plugins/Experimental/ChaosEditor/Source/FractureEditor/Private/FractureEditorMode.cpp
+++ b/Engine/Plugins/Experimental/ChaosEditor/Source/FractureEditor/Private/FractureEditorMode.cpp
@@ -583,7 +583,11 @@ void UFractureEditorMode::OnActorSelectionChanged(const TArray<UObject*>& NewSel
 			constexpr bool bForceUpdate = true; // Force the bone selection and highlight to refresh so bone colors reflect the selection
 			FScopedColorEdit ShowBoneColorsEdit(GeometryCollectionComponent, bForceUpdate);
 			ShowBoneColorsEdit.SetEnableBoneSelection(true);
-			// ShowBoneColorsEdit.SetLevelViewMode(ViewLevel);
+			if (ViewLevel != ShowBoneColorsEdit.GetViewLevel())
+			{
+				ShowBoneColorsEdit.SetLevelViewMode(ViewLevel);
+				ShowBoneColorsEdit.FilterSelectionToLevel();
+			}
 			ShowBoneColorsEdit.Sanitize(); // Clean any stale data (e.g. due to the geometry being edited via a different component)
 
 			NewGeomSelection.Add(GeometryCollectionComponent);
