fix issue where group fracture setting would disappear after using a fracture tool w/ a gizmo:
 - stop fully hiding the setting when it's disabled, so it's clearer what is happening in the UI
 - reset whether the setting is enabled on fracture tool setup, so state from one tool doesn't affect another tool

diff --git a/Engine/Plugins/Experimental/ChaosEditor/Source/FractureEditor/Private/FractureToolBrick.h b/Engine/Plugins/Experimental/ChaosEditor/Source/FractureEditor/Private/FractureToolBrick.h
--- a/Engine/Plugins/Experimental/ChaosEditor/Source/FractureEditor/Private/FractureToolBrick.h
+++ b/Engine/Plugins/Experimental/ChaosEditor/Source/FractureEditor/Private/FractureToolBrick.h
@@ -100,6 +100,7 @@ public:
 
 	virtual void ConfigureCutterSettings() override
 	{
+		CutterSettings->bGroupFractureToggleEnabled = true;
 		CutterSettings->bNoisePreviewToggleEnabled = false; // TODO: support noise preview for bricks
 	}
 
diff --git a/Engine/Plugins/Experimental/ChaosEditor/Source/FractureEditor/Private/FractureToolCutter.h b/Engine/Plugins/Experimental/ChaosEditor/Source/FractureEditor/Private/FractureToolCutter.h
--- a/Engine/Plugins/Experimental/ChaosEditor/Source/FractureEditor/Private/FractureToolCutter.h
+++ b/Engine/Plugins/Experimental/ChaosEditor/Source/FractureEditor/Private/FractureToolCutter.h
@@ -75,7 +75,7 @@ public:
 	float ChanceToFracture;
 
 	/** Generate a fracture pattern across all selected meshes.  */
-	UPROPERTY(EditAnywhere, Category = CommonFracture, meta = (EditCondition = "bGroupFractureToggleEnabled", HideEditConditionToggle, EditConditionHides, DisplayName = "Group Fracture"))
+	UPROPERTY(EditAnywhere, Category = CommonFracture, meta = (EditCondition = "bGroupFractureToggleEnabled", HideEditConditionToggle, DisplayName = "Group Fracture"))
 	bool bGroupFracture;
 
 	// This flag allows tools to disable the above bGroupFracture option if/when it is not applicable
@@ -242,6 +242,7 @@ public:
 	// Set tool-specific defaults for the cutter settings (e.g., disable grout/noise if needed)
 	virtual void ConfigureCutterSettings()
 	{
+		CutterSettings->bGroupFractureToggleEnabled = true;
 		CutterSettings->bDrawSitesToggleEnabled = true;
 		CutterSettings->bNoisePreviewToggleEnabled = true;
 		CutterSettings->bNoisePreviewHasScale = false;
diff --git a/Engine/Plugins/Experimental/ChaosEditor/Source/FractureEditor/Private/FractureToolMeshCut.h b/Engine/Plugins/Experimental/ChaosEditor/Source/FractureEditor/Private/FractureToolMeshCut.h
--- a/Engine/Plugins/Experimental/ChaosEditor/Source/FractureEditor/Private/FractureToolMeshCut.h
+++ b/Engine/Plugins/Experimental/ChaosEditor/Source/FractureEditor/Private/FractureToolMeshCut.h
@@ -111,6 +111,7 @@ public:
 
 	virtual void ConfigureCutterSettings() override
 	{
+		CutterSettings->bGroupFractureToggleEnabled = true;
 		CutterSettings->bDrawSitesToggleEnabled = false;
 		CutterSettings->bNoisePreviewToggleEnabled = false;
 		CutterSettings->bNoisePreviewHasScale = true;
diff --git a/Engine/Plugins/Experimental/ChaosEditor/Source/FractureEditor/Private/FractureToolPlaneCut.h b/Engine/Plugins/Experimental/ChaosEditor/Source/FractureEditor/Private/FractureToolPlaneCut.h
--- a/Engine/Plugins/Experimental/ChaosEditor/Source/FractureEditor/Private/FractureToolPlaneCut.h
+++ b/Engine/Plugins/Experimental/ChaosEditor/Source/FractureEditor/Private/FractureToolPlaneCut.h
@@ -65,6 +65,7 @@ public:
 
 	virtual void ConfigureCutterSettings() override
 	{
+		CutterSettings->bGroupFractureToggleEnabled = true;
 		CutterSettings->bDrawSitesToggleEnabled = false;
 		CutterSettings->bNoisePreviewHasScale = true;
 	}
diff --git a/Engine/Plugins/Experimental/ChaosEditor/Source/FractureEditor/Private/FractureToolSlice.h b/Engine/Plugins/Experimental/ChaosEditor/Source/FractureEditor/Private/FractureToolSlice.h
--- a/Engine/Plugins/Experimental/ChaosEditor/Source/FractureEditor/Private/FractureToolSlice.h
+++ b/Engine/Plugins/Experimental/ChaosEditor/Source/FractureEditor/Private/FractureToolSlice.h
@@ -71,6 +71,7 @@ public:
 
 	virtual void ConfigureCutterSettings() override
 	{
+		CutterSettings->bGroupFractureToggleEnabled = true;
 		CutterSettings->bDrawSitesToggleEnabled = false;
 		CutterSettings->bNoisePreviewToggleEnabled = false; // TODO: support noise preview for slices
 	}
