set random seed earlier on fracture context so it is used for selection 'chance to cut' downsampling as well

diff --git a/Engine/Plugins/Experimental/ChaosEditor/Source/FractureEditor/Private/FractureToolCutter.cpp b/Engine/Plugins/Experimental/ChaosEditor/Source/FractureEditor/Private/FractureToolCutter.cpp
--- a/Engine/Plugins/Experimental/ChaosEditor/Source/FractureEditor/Private/FractureToolCutter.cpp
+++ b/Engine/Plugins/Experimental/ChaosEditor/Source/FractureEditor/Private/FractureToolCutter.cpp
@@ -167,6 +167,8 @@ TArray<FFractureToolContext> UFractureToolCutterBase::GetFractureToolContexts()
 
 			// Generate a context for each selected node
 			FFractureToolContext FullSelection(GeometryCollectionComponent);
+			int32 FullSelectionSeed = CutterSettings->RandomSeed > -1 ? CutterSettings->RandomSeed : DefaultRandomSeed;
+			FullSelection.SetSeed(FullSelectionSeed);
 			FullSelection.ConvertSelectionToRigidNodes();
 			FullSelection.RandomReduceSelection(CutterSettings->ChanceToFracture);
 
@@ -190,8 +192,6 @@ TArray<FFractureToolContext> UFractureToolCutterBase::GetFractureToolContexts()
 
 			if (CutterSettings->bGroupFracture)
 			{
-				FullSelection.SetSeed(CutterSettings->RandomSeed > -1 ? CutterSettings->RandomSeed : DefaultRandomSeed);
-
 				FBox Bounds(ForceInit);
 				for (int32 BoneIndex : FullSelection.GetSelection())
 				{
@@ -214,7 +214,7 @@ TArray<FFractureToolContext> UFractureToolCutterBase::GetFractureToolContexts()
 					TArray<int32> Selection;
 					Selection.Add(Index);
 					FractureContext.SetSelection(Selection);
-					FractureContext.SetSeed(CutterSettings->RandomSeed > -1 ? CutterSettings->RandomSeed + Index : DefaultRandomSeed + Index);
+					FractureContext.SetSeed(FullSelectionSeed + Index);
 					FractureContext.SetBounds(BoundsToBone[Index]);
 				}
 			}
