fix tooltip that describes geometry collection bone 'size' as based on bounds side length; from looking at the code, it appears to be cube root of volume

diff --git a/Engine/Plugins/Experimental/ChaosEditor/Source/FractureEditor/Private/SGeometryCollectionOutliner.cpp b/Engine/Plugins/Experimental/ChaosEditor/Source/FractureEditor/Private/SGeometryCollectionOutliner.cpp
--- a/Engine/Plugins/Experimental/ChaosEditor/Source/FractureEditor/Private/SGeometryCollectionOutliner.cpp
+++ b/Engine/Plugins/Experimental/ChaosEditor/Source/FractureEditor/Private/SGeometryCollectionOutliner.cpp
@@ -389,7 +389,7 @@ void SGeometryCollectionOutliner::RegenerateHeader()
 		HeaderRowWidget->AddColumn(
 			SHeaderRow::Column(SGeometryCollectionOutlinerColumnID::Volume)
 			.DefaultLabel(LOCTEXT("GCOutliner_Column_Volume", "Size"))
-			.DefaultTooltip(LOCTEXT("GCOutliner_Column_Volume_ToolTip", "Side length of the bounding cube (in cm)"))
+			.DefaultTooltip(LOCTEXT("GCOutliner_Column_Volume_ToolTip", "Cube root of bone volume (in cm)"))
 			.HAlignHeader(EHorizontalAlignment::HAlign_Center)
 			.FillWidth(CustomFillWidth)
 		);
