fix crash on deleting a geometry collection while a fracture tool has an active noise preview

diff --git a/Engine/Plugins/Experimental/ChaosEditor/Source/FractureEditor/Private/FractureToolCutter.h b/Engine/Plugins/Experimental/ChaosEditor/Source/FractureEditor/Private/FractureToolCutter.h
--- a/Engine/Plugins/Experimental/ChaosEditor/Source/FractureEditor/Private/FractureToolCutter.h
+++ b/Engine/Plugins/Experimental/ChaosEditor/Source/FractureEditor/Private/FractureToolCutter.h
@@ -412,6 +412,7 @@ protected:
 private:
 	UPROPERTY()
 	TArray<TObjectPtr<ULineSetComponent>> VoronoiLineSets;
+	UPROPERTY()
 	TArray<TObjectPtr<UDynamicMeshComponent>> VoronoiNoisePreviews;
 
 
