Fracture editor : fix the tool palette labels to be localized

diff --git a/Engine/Plugins/Experimental/ChaosEditor/Source/FractureEditor/Private/FractureEditorModeToolkit.cpp b/Engine/Plugins/Experimental/ChaosEditor/Source/FractureEditor/Private/FractureEditorModeToolkit.cpp
--- a/Engine/Plugins/Experimental/ChaosEditor/Source/FractureEditor/Private/FractureEditorModeToolkit.cpp
+++ b/Engine/Plugins/Experimental/ChaosEditor/Source/FractureEditor/Private/FractureEditorModeToolkit.cpp
@@ -936,10 +936,64 @@ void FFractureEditorModeToolkit::OnObjectPostEditChange( UObject* Object, FPrope
 	}
 }
 
-const TArray<FName> FFractureEditorModeToolkit::PaletteNames = { FName(TEXT("Generate")), FName(TEXT("Select")), FName(TEXT("Fracture")), FName(TEXT("Edit")), FName(TEXT("Cluster")), FName(TEXT("Embed")), FName(TEXT("Utilities")) };
+namespace FractureEditorModeToolkit
+{
+	static const FName GenerateTools = TEXT("Generate");
+	static const FName SelectTools = TEXT("Select");
+	static const FName FractureTools = TEXT("Fracture");
+	static const FName EditTools = TEXT("Edit");
+	static const FName ClusterTools = TEXT("Cluster");
+	static const FName EmbedTools = TEXT("Embed");
+	static const FName UtilitiesTools = TEXT("Utilities");
+	static const FName FavoritesTools = TEXT("Favorites");
+}
+
+const TArray<FName> FFractureEditorModeToolkit::PaletteNames = 
+{ 
+	FractureEditorModeToolkit::GenerateTools,
+	FractureEditorModeToolkit::SelectTools,
+	FractureEditorModeToolkit::FractureTools,
+	FractureEditorModeToolkit::EditTools,
+	FractureEditorModeToolkit::ClusterTools,
+	FractureEditorModeToolkit::EmbedTools,
+	FractureEditorModeToolkit::UtilitiesTools,
+	FractureEditorModeToolkit::FavoritesTools,
+};
 
 FText FFractureEditorModeToolkit::GetToolPaletteDisplayName(FName Palette) const
 { 
+	if (Palette == FractureEditorModeToolkit::GenerateTools)
+	{
+		return LOCTEXT("FractureEditorModeTookit_GenerateTools", "Generate");
+	}
+	if (Palette == FractureEditorModeToolkit::SelectTools)
+	{
+		return LOCTEXT("FractureEditorModeTookit_SelectTools", "Select");
+	}
+	if (Palette == FractureEditorModeToolkit::FractureTools)
+	{
+		return LOCTEXT("FractureEditorModeTookit_FractureTools", "Fracture");
+	}
+	if (Palette == FractureEditorModeToolkit::EditTools)
+	{
+		return LOCTEXT("FractureEditorModeTookit_EditTools", "Edit");
+	}
+	if (Palette == FractureEditorModeToolkit::ClusterTools)
+	{
+		return LOCTEXT("FractureEditorModeTookit_ClusterTools", "Cluster");
+	}
+	if (Palette == FractureEditorModeToolkit::EmbedTools)
+	{
+		return LOCTEXT("FractureEditorModeTookit_EmbedTools", "Embed");
+	}
+	if (Palette == FractureEditorModeToolkit::UtilitiesTools)
+	{
+		return LOCTEXT("FractureEditorModeTookit_UtilitiesTools", "Utilities");
+	}
+	if (Palette == FractureEditorModeToolkit::FavoritesTools)
+	{
+		return LOCTEXT("FractureEditorModeTookit_FavoritesTools", "Favorites");
+	}
 	return FText::FromName(Palette);
 }
 
@@ -948,11 +1002,11 @@ void FFractureEditorModeToolkit::SetInitialPalette()
 	// Start in Select Palette if GeometryCollection is selected.
 	if (IsGeometryCollectionSelected())
 	{
-		SetCurrentPalette(TEXT("Select"));
+		SetCurrentPalette(FractureEditorModeToolkit::SelectTools);
 	}
 	else
 	{
-		SetCurrentPalette(TEXT("Generate"));
+		SetCurrentPalette(FractureEditorModeToolkit::GenerateTools);
 	}
 }
 
@@ -986,7 +1040,7 @@ void FFractureEditorModeToolkit::GetToolPaletteNames(TArray<FName>& PaletteNames
 	// if user has provided a list of favorite tools, add that palette to the list
 	if (UISettings->ToolFavorites.Num() > 0)
 	{
-		PaletteNamesOut.Insert(FName(TEXT("Favorites")), 0);
+		PaletteNamesOut.Insert(FractureEditorModeToolkit::FavoritesTools, 0);
 	}
 }
 
@@ -994,7 +1048,7 @@ void FFractureEditorModeToolkit::BuildToolPalette(FName PaletteIndex, class FToo
 {
 	const FFractureEditorCommands& Commands = FFractureEditorCommands::Get();
 
-	if (PaletteIndex == TEXT("Favorites"))
+	if (PaletteIndex == FractureEditorModeToolkit::FavoritesTools)
 	{
 		UFractureModeCustomizationSettings* UISettings = GetMutableDefault<UFractureModeCustomizationSettings>();
 
@@ -1013,12 +1067,12 @@ void FFractureEditorModeToolkit::BuildToolPalette(FName PaletteIndex, class FToo
 			}
 		}
 	}
-	else if (PaletteIndex == TEXT("Generate"))
+	else if (PaletteIndex == FractureEditorModeToolkit::GenerateTools)
 	{
 		ToolbarBuilder.AddToolBarButton(Commands.GenerateAsset);
 		ToolbarBuilder.AddToolBarButton(Commands.ResetAsset);
 	}
-	else if (PaletteIndex == TEXT("Select"))
+	else if (PaletteIndex == FractureEditorModeToolkit::SelectTools)
 	{
 		ToolbarBuilder.AddToolBarButton(Commands.SelectAll);
 		ToolbarBuilder.AddToolBarButton(Commands.SelectInvert);
@@ -1032,7 +1086,7 @@ void FFractureEditorModeToolkit::BuildToolPalette(FName PaletteIndex, class FToo
 		ToolbarBuilder.AddToolBarButton(Commands.SelectClusters);
 		ToolbarBuilder.AddToolBarButton(Commands.SelectCustom);
 	}
-	else if (PaletteIndex == TEXT("Fracture"))
+	else if (PaletteIndex == FractureEditorModeToolkit::FractureTools)
 	{
 		ToolbarBuilder.AddToolBarButton(Commands.Uniform);
 		ToolbarBuilder.AddToolBarButton(Commands.Clustered);
@@ -1043,7 +1097,7 @@ void FFractureEditorModeToolkit::BuildToolPalette(FName PaletteIndex, class FToo
 		ToolbarBuilder.AddToolBarButton(Commands.Mesh);
 		ToolbarBuilder.AddToolBarButton(Commands.CustomVoronoi);
 	}
-	else if (PaletteIndex == TEXT("Edit"))
+	else if (PaletteIndex == FractureEditorModeToolkit::EditTools)
 	{
 		ToolbarBuilder.AddToolBarButton(Commands.DeleteBranch);
 		ToolbarBuilder.AddToolBarButton(Commands.Hide);
@@ -1051,7 +1105,7 @@ void FFractureEditorModeToolkit::BuildToolPalette(FName PaletteIndex, class FToo
 		ToolbarBuilder.AddToolBarButton(Commands.MergeSelected);
 		//ToolbarBuilder.AddToolBarButton(Commands.SplitSelected); // Split tool intentionally disabled; prefer the 'split' option on import instead.
 	}
-	else if (PaletteIndex == TEXT("Cluster"))
+	else if (PaletteIndex == FractureEditorModeToolkit::ClusterTools)
 	{
 		ToolbarBuilder.AddToolBarButton(Commands.AutoCluster);
 		ToolbarBuilder.AddToolBarButton(Commands.ClusterMagnet);
@@ -1061,13 +1115,13 @@ void FFractureEditorModeToolkit::BuildToolPalette(FName PaletteIndex, class FToo
 		ToolbarBuilder.AddToolBarButton(Commands.MoveUp);
 		ToolbarBuilder.AddToolBarButton(Commands.ClusterMerge);
 	}
-	else if (PaletteIndex == TEXT("Embed"))
+	else if (PaletteIndex == FractureEditorModeToolkit::EmbedTools)
 	{
 		ToolbarBuilder.AddToolBarButton(Commands.AddEmbeddedGeometry);
 		ToolbarBuilder.AddToolBarButton(Commands.AutoEmbedGeometry);
 		ToolbarBuilder.AddToolBarButton(Commands.FlushEmbeddedGeometry);
 	}
-	else if (PaletteIndex == TEXT("Utilities"))
+	else if (PaletteIndex == FractureEditorModeToolkit::UtilitiesTools)
 	{
 		ToolbarBuilder.AddToolBarButton(Commands.AutoUV);
 		ToolbarBuilder.AddToolBarButton(Commands.Materials);
