Fracture Editor : localize the selection info and add the total number of item in the outiner

diff --git a/Engine/Plugins/Experimental/ChaosEditor/Source/FractureEditor/Private/FractureEditorModeToolkit.cpp b/Engine/Plugins/Experimental/ChaosEditor/Source/FractureEditor/Private/FractureEditorModeToolkit.cpp
--- a/Engine/Plugins/Experimental/ChaosEditor/Source/FractureEditor/Private/FractureEditorModeToolkit.cpp
+++ b/Engine/Plugins/Experimental/ChaosEditor/Source/FractureEditor/Private/FractureEditorModeToolkit.cpp
@@ -2194,8 +2194,7 @@ void FFractureEditorModeToolkit::OnHistogramBoneSelectionChanged(UGeometryCollec
 
 FText FFractureEditorModeToolkit::GetSelectionInfo() const
 {
-	FString Buffer = FString::Printf(TEXT("Selected: %d"), OutlinerView->GetBoneSelectionCount());
-	return FText::FromString(Buffer);
+	return FText::Format(LOCTEXT("FractureEditorModeToolkit_SelectionInfo", "Selected: {0} / {1}"), OutlinerView->GetBoneSelectionCount(), OutlinerView->GetBoneDisplayedCount());
 }
 
 
diff --git a/Engine/Plugins/Experimental/ChaosEditor/Source/FractureEditor/Private/SGeometryCollectionOutliner.cpp b/Engine/Plugins/Experimental/ChaosEditor/Source/FractureEditor/Private/SGeometryCollectionOutliner.cpp
--- a/Engine/Plugins/Experimental/ChaosEditor/Source/FractureEditor/Private/SGeometryCollectionOutliner.cpp
+++ b/Engine/Plugins/Experimental/ChaosEditor/Source/FractureEditor/Private/SGeometryCollectionOutliner.cpp
@@ -546,6 +546,11 @@ int32 SGeometryCollectionOutliner::GetBoneSelectionCount() const
 	return TreeView->GetSelectedItems().Num();
 }
 
+int32 SGeometryCollectionOutliner::GetBoneDisplayedCount() const
+{
+	return TreeView->GetNumGeneratedChildren();
+}
+
 void SGeometryCollectionOutliner::SetBoneSelection(UGeometryCollectionComponent* RootComponent, const TArray<int32>& InSelection, bool bClearCurrentSelection, int32 FocusBoneIdx)
 {
 	TGuardValue<bool> ExternalSelectionGuard(bPerformingSelection, true);
diff --git a/Engine/Plugins/Experimental/ChaosEditor/Source/FractureEditor/Private/SGeometryCollectionOutliner.h b/Engine/Plugins/Experimental/ChaosEditor/Source/FractureEditor/Private/SGeometryCollectionOutliner.h
--- a/Engine/Plugins/Experimental/ChaosEditor/Source/FractureEditor/Private/SGeometryCollectionOutliner.h
+++ b/Engine/Plugins/Experimental/ChaosEditor/Source/FractureEditor/Private/SGeometryCollectionOutliner.h
@@ -330,6 +330,7 @@ public:
 	void SetComponents(const TArray<UGeometryCollectionComponent*>& InNewComponents);
 	void SetBoneSelection(UGeometryCollectionComponent* RootComponent, const TArray<int32>& InSelection, bool bClearCurrentSelection, int32 FocusBoneIdx = -1);
 	int32 GetBoneSelectionCount() const;
+	int32 GetBoneDisplayedCount() const;
 
 	void SetInitialDynamicState(int32 InDynamicState);
 	void SetAnchored(bool bAnchored);
