Fracture Editor : display the dataflow variable overrides to the setting panel

diff --git a/Engine/Plugins/Experimental/ChaosEditor/Source/FractureEditor/FractureEditor.Build.cs b/Engine/Plugins/Experimental/ChaosEditor/Source/FractureEditor/FractureEditor.Build.cs
--- a/Engine/Plugins/Experimental/ChaosEditor/Source/FractureEditor/FractureEditor.Build.cs
+++ b/Engine/Plugins/Experimental/ChaosEditor/Source/FractureEditor/FractureEditor.Build.cs
@@ -40,6 +40,7 @@ public class FractureEditor : ModuleRules
 				"PropertyEditor",
 				"GeometryCore",
 				"DataflowEngine",
+				"DataflowEditor",
 				"DynamicMesh",
 				"MeshDescription",
 				"StaticMeshDescription",
diff --git a/Engine/Plugins/Experimental/ChaosEditor/Source/FractureEditor/Private/FractureEditorModeToolkit.cpp b/Engine/Plugins/Experimental/ChaosEditor/Source/FractureEditor/Private/FractureEditorModeToolkit.cpp
--- a/Engine/Plugins/Experimental/ChaosEditor/Source/FractureEditor/Private/FractureEditorModeToolkit.cpp
+++ b/Engine/Plugins/Experimental/ChaosEditor/Source/FractureEditor/Private/FractureEditorModeToolkit.cpp
@@ -7,12 +7,14 @@
 #include "EditorModeManager.h"
 
 #include "Engine/Selection.h"
+#include "Dataflow/DataflowInstanceDetails.h"
 #include "FractureEditorMode.h"
 #include "Framework/MultiBox/MultiBoxBuilder.h"
 #include "Framework/MultiBox/MultiBoxDefs.h"
 #include "GeometryCollection/GeometryCollectionObject.h"
 #include "IDetailsView.h"
 #include "IDetailRootObjectCustomization.h"
+#include "IStructureDetailsView.h"
 #include "PropertyEditorModule.h"
 #include "UObject/UObjectIterator.h"
 #include "Widgets/Docking/SDockTab.h"
@@ -247,6 +249,73 @@ void FFractureEditorModeToolkit::UpdateAssetPanelFromSettings()
 	}
 }
 
+void FFractureEditorModeToolkit::CreateVariableOverrideDetailView()
+{
+	FDetailsViewArgs DetailsViewArgs;
+	{
+		DetailsViewArgs.bAllowSearch = false;
+		DetailsViewArgs.NameAreaSettings = FDetailsViewArgs::HideNameArea;
+		DetailsViewArgs.bHideSelectionTip = true;
+		DetailsViewArgs.DefaultsOnlyVisibility = EEditDefaultsOnlyNodeVisibility::Automatic;
+		DetailsViewArgs.bShowOptions = false;
+		DetailsViewArgs.bAllowMultipleTopLevelObjects = true;
+		DetailsViewArgs.bShowKeyablePropertiesOption = false;
+		DetailsViewArgs.bShowModifiedPropertiesOption = false;
+		DetailsViewArgs.bAllowFavoriteSystem = false;
+		DetailsViewArgs.bShowAnimatedPropertiesOption = false;
+	}
+
+	FStructureDetailsViewArgs StructureViewArgs;
+	{
+		StructureViewArgs.bShowObjects = true;
+		StructureViewArgs.bShowAssets = true;
+		StructureViewArgs.bShowClasses = true;
+		StructureViewArgs.bShowInterfaces = true;
+	}
+
+	FPropertyEditorModule& PropertyEditorModule = FModuleManager::Get().GetModuleChecked<FPropertyEditorModule>("PropertyEditor");
+	OverridesDetailsView = PropertyEditorModule.CreateStructureDetailView(DetailsViewArgs, StructureViewArgs, nullptr);
+
+	auto MakeDataflowInstanceDetailCustomizationLambda = []()
+		{
+			constexpr bool bOnlyShowVariableOverrides = true;
+			return MakeShared<FDataflowInstanceDetailCustomization>(bOnlyShowVariableOverrides);
+		};
+	FOnGetDetailCustomizationInstance DataflowInstanceDetailsCustomizationInstance = FOnGetDetailCustomizationInstance::CreateLambda(MakeDataflowInstanceDetailCustomizationLambda);
+	OverridesDetailsView->GetDetailsView()->RegisterInstancedCustomPropertyLayout(FDataflowInstance::StaticStruct(), DataflowInstanceDetailsCustomizationInstance);
+
+	OverridesDetailsView->GetOnFinishedChangingPropertiesDelegate().AddRaw(this, &FFractureEditorModeToolkit::OverridesDetailsViewFinishedChangingProperties);
+
+	RefreshVariableOverrideDetailView(nullptr);
+}
+
+void FFractureEditorModeToolkit::RefreshVariableOverrideDetailView(const UGeometryCollection* RestCollection)
+{
+	if (RestCollection == nullptr)
+	{
+		OverridesDetailsView->GetDetailsView()->SetObject(nullptr);
+		OverridesDetailsView->SetStructureData(nullptr);
+		return;
+	}
+	// we only get a const collection from the component, but the fracture editor needs to be able to edit the variables 
+	UGeometryCollection* EditableRestCollection = const_cast<UGeometryCollection*>(RestCollection);
+	if (const IDataflowInstanceInterface* Interface = EditableRestCollection)
+	{
+		OverridesDetailsView->GetDetailsView()->SetObject(EditableRestCollection);
+		TSharedPtr<FStructOnScope> StructOnScope = Interface->GetDataflowInstance().MakeStructOnScope();
+		OverridesDetailsView->SetStructureData(StructOnScope);
+	}
+}
+
+void FFractureEditorModeToolkit::OverridesDetailsViewFinishedChangingProperties(const FPropertyChangedEvent& InPropertyChangedEvent)
+{
+	// get the displayed asset from the settings
+	UFractureSettings* FractureSettings = GetMutableDefault<UFractureSettings>();
+	if (FractureSettings->RestCollection.IsValid())
+	{
+		FractureSettings->RestCollection->GetDataflowInstance().UpdateOwnerAsset();
+	}
+}
 
 void FFractureEditorModeToolkit::Init(const TSharedPtr<IToolkitHost>& InitToolkitHost, TWeakObjectPtr<UEdMode> InOwningMode)
 {
@@ -293,6 +362,8 @@ void FFractureEditorModeToolkit::Init(const TSharedPtr<IToolkitHost>& InitToolki
 	EditModule.RegisterCustomClassLayout("OutlinerSettings", FOnGetDetailCustomizationInstance::CreateStatic(&FOutlinerSettingsCustomization::MakeInstance, this));
 	OutlinerDetailsView->SetObject(GetMutableDefault<UOutlinerSettings>());
 
+	CreateVariableOverrideDetailView();
+
 	AssetLocationModes.Reset();
 	AssetLocationModes.Add(MakeShared<FString>(TEXT("Source Asset Folder")));
 	AssetLocationModes.Add(MakeShared<FString>(TEXT("Last Used Folder")));
@@ -385,6 +456,15 @@ void FFractureEditorModeToolkit::Init(const TSharedPtr<IToolkitHost>& InitToolki
 						DetailsView.ToSharedRef()
 					]
 				]
+				+ SVerticalBox::Slot()
+				.FillHeight(1.0)
+				[
+					SNew(SScrollBox)
+					+ SScrollBox::Slot()
+					[
+						OverridesDetailsView->GetWidget().ToSharedRef()
+					]
+				]
 				+SVerticalBox::Slot()
 				.AutoHeight()
 				[
@@ -1797,6 +1877,8 @@ void FFractureEditorModeToolkit::SetOutlinerComponents(const TArray<UGeometryCol
 		break;
 	}
 
+	RefreshVariableOverrideDetailView(FractureSettings->RestCollection.Get());
+
 	if (ActiveTool != nullptr)
 	{
 		ActiveTool->SelectedBonesChanged();
diff --git a/Engine/Plugins/Experimental/ChaosEditor/Source/FractureEditor/Public/FractureEditorModeToolkit.h b/Engine/Plugins/Experimental/ChaosEditor/Source/FractureEditor/Public/FractureEditorModeToolkit.h
--- a/Engine/Plugins/Experimental/ChaosEditor/Source/FractureEditor/Public/FractureEditorModeToolkit.h
+++ b/Engine/Plugins/Experimental/ChaosEditor/Source/FractureEditor/Public/FractureEditorModeToolkit.h
@@ -15,9 +15,11 @@
 
 class IDetailsView;
 class IPropertyHandle;
+class IStructureDetailsView;
 class SScrollBox;
 class FFractureToolContext;
 class FGeometryCollection;
+struct FPropertyChangedEvent;
 class SGeometryCollectionOutliner;
 class SGeometryCollectionHistogram;
 class SGeometryCollectionStatistics;
@@ -241,7 +243,11 @@ private:
 	void OnProjectSettingsModified();
 
 	void UpdateOutlinerHeader();
-	
+
+	void CreateVariableOverrideDetailView();
+	void RefreshVariableOverrideDetailView(const UGeometryCollection* RestCollection);
+	void OverridesDetailsViewFinishedChangingProperties(const FPropertyChangedEvent& InPropertyChangedEvent);
+
 private:
 	TObjectPtr<UFractureModalTool> ActiveTool;
 
@@ -251,7 +257,7 @@ private:
 	FDelegateHandle ProjectSettingsModifiedHandle;
 
 	TSharedPtr<IDetailsView> DetailsView;
-	TSharedPtr<IDetailsView> ViewSettingsDetailsView;
+	TSharedPtr<IStructureDetailsView> OverridesDetailsView;
 	TSharedPtr<IDetailsView> HistogramDetailsView;
 	TSharedPtr<IDetailsView> OutlinerDetailsView;
 	TSharedPtr<SWidget> ToolkitWidget;
diff --git a/Engine/Plugins/Experimental/Dataflow/Source/DataflowEditor/Public/Dataflow/DataflowInstanceDetails.h b/Engine/Plugins/Experimental/Dataflow/Source/DataflowEditor/Public/Dataflow/DataflowInstanceDetails.h
--- a/Engine/Plugins/Experimental/Dataflow/Source/DataflowEditor/Public/Dataflow/DataflowInstanceDetails.h
+++ b/Engine/Plugins/Experimental/Dataflow/Source/DataflowEditor/Public/Dataflow/DataflowInstanceDetails.h
@@ -28,7 +28,7 @@ class FDataflowInstanceDetailCustomization : public IDetailCustomization
 {
 public:
 	/** Makes a new instance of this detail layout class for a specific detail view requesting it */
-	FDataflowInstanceDetailCustomization(bool bOnlyShowVariableOverrides);
+	DATAFLOWEDITOR_API FDataflowInstanceDetailCustomization(bool bOnlyShowVariableOverrides);
 
 	/** IDetailCustomization interface */
 	virtual void CustomizeDetails(IDetailLayoutBuilder& DetailBuilder) override;
diff --git a/Engine/Source/Runtime/Experimental/Dataflow/Engine/Private/Dataflow/DataflowInstance.cpp b/Engine/Source/Runtime/Experimental/Dataflow/Engine/Private/Dataflow/DataflowInstance.cpp
--- a/Engine/Source/Runtime/Experimental/Dataflow/Engine/Private/Dataflow/DataflowInstance.cpp
+++ b/Engine/Source/Runtime/Experimental/Dataflow/Engine/Private/Dataflow/DataflowInstance.cpp
@@ -2,6 +2,7 @@
 
 #include "Dataflow/DataflowInstance.h"
 #include "Dataflow/DataflowObject.h"
+#include "Dataflow/DataflowObjectInterface.h"
 
 namespace UE::Dataflow::InstanceUtils
 {
@@ -359,3 +360,27 @@ void FDataflowInstance::SyncVariables()
 {
 	VariableOverrides.SyncVariables();
 }
+
+bool FDataflowInstance::UpdateOwnerAsset() const
+{
+	if (DataflowAsset && DataflowAsset->Dataflow && Owner)
+	{
+		if (const TSharedPtr<FDataflowNode> Node = DataflowAsset->Dataflow->FindFilteredNode(FDataflowTerminalNode::StaticType(), DataflowTerminal))
+		{
+			if (const FDataflowTerminalNode* TerminalNode = Node->AsType<const FDataflowTerminalNode>())
+			{
+				UE::Dataflow::FEngineContext Context(Owner.Get());
+				// Note: If the node is deactivated and has any outputs, then these outputs might still need to be forwarded.
+				//       Therefore the Evaluate method has to be called for whichever value of bActive.
+				//       This however isn't the case of SetAssetValue() for which the active state needs to be checked before the call.
+				TerminalNode->Evaluate(Context);
+				if (TerminalNode->IsActive() && Owner)
+				{
+					TerminalNode->SetAssetValue(Owner, Context);
+					return true;
+				}
+			}
+		}
+	}
+	return false;
+}
\ No newline at end of file
diff --git a/Engine/Source/Runtime/Experimental/Dataflow/Engine/Public/Dataflow/DataflowInstance.h b/Engine/Source/Runtime/Experimental/Dataflow/Engine/Public/Dataflow/DataflowInstance.h
--- a/Engine/Source/Runtime/Experimental/Dataflow/Engine/Public/Dataflow/DataflowInstance.h
+++ b/Engine/Source/Runtime/Experimental/Dataflow/Engine/Public/Dataflow/DataflowInstance.h
@@ -112,6 +112,12 @@ public:
 
 	DATAFLOWENGINE_API void SyncVariables();
 
+	/** 
+	* Update the asset by re-evaluating the dataflow
+	* return true if the dataflow has been proeprly evaluated and the asset updated
+	*/
+	DATAFLOWENGINE_API bool UpdateOwnerAsset() const;
+
 #if WITH_EDITOR
 	DATAFLOWENGINE_API TSharedPtr<FStructOnScope> MakeStructOnScope() const;
 #endif 
