Fracture tool : Fix height ratio between the settings and the variable overrides

diff --git a/Engine/Plugins/Experimental/ChaosEditor/Source/FractureEditor/Private/FractureEditorModeToolkit.cpp b/Engine/Plugins/Experimental/ChaosEditor/Source/FractureEditor/Private/FractureEditorModeToolkit.cpp
--- a/Engine/Plugins/Experimental/ChaosEditor/Source/FractureEditor/Private/FractureEditorModeToolkit.cpp
+++ b/Engine/Plugins/Experimental/ChaosEditor/Source/FractureEditor/Private/FractureEditorModeToolkit.cpp
@@ -448,7 +448,7 @@ void FFractureEditorModeToolkit::Init(const TSharedPtr<IToolkitHost>& InitToolki
 			[
 				SNew(SVerticalBox)
 				+SVerticalBox::Slot()
-				.FillHeight(1.0)
+				.AutoHeight()
 				[
 					SNew(SScrollBox)
 					+SScrollBox::Slot()
