Fracture editor : Dataflow override imporvements
- Separate the global fracture settings from the tools settings to allow the dataflow override to stay below the global settings
- Do not re-execute the dataflow in real time as the override are updated and instead let the user press the update button ( works also for non-variable driven dataflow)
- asset with no ataflow will see the button  disabled

diff --git a/Engine/Plugins/Experimental/ChaosEditor/Source/FractureEditor/Private/FractureEditorModeToolkit.cpp b/Engine/Plugins/Experimental/ChaosEditor/Source/FractureEditor/Private/FractureEditorModeToolkit.cpp
--- a/Engine/Plugins/Experimental/ChaosEditor/Source/FractureEditor/Private/FractureEditorModeToolkit.cpp
+++ b/Engine/Plugins/Experimental/ChaosEditor/Source/FractureEditor/Private/FractureEditorModeToolkit.cpp
@@ -284,8 +284,6 @@ void FFractureEditorModeToolkit::CreateVariableOverrideDetailView()
 	FOnGetDetailCustomizationInstance DataflowInstanceDetailsCustomizationInstance = FOnGetDetailCustomizationInstance::CreateLambda(MakeDataflowInstanceDetailCustomizationLambda);
 	OverridesDetailsView->GetDetailsView()->RegisterInstancedCustomPropertyLayout(FDataflowInstance::StaticStruct(), DataflowInstanceDetailsCustomizationInstance);
 
-	OverridesDetailsView->GetOnFinishedChangingPropertiesDelegate().AddRaw(this, &FFractureEditorModeToolkit::OverridesDetailsViewFinishedChangingProperties);
-
 	RefreshVariableOverrideDetailView(nullptr);
 }
 
@@ -307,14 +305,33 @@ void FFractureEditorModeToolkit::RefreshVariableOverrideDetailView(const UGeomet
 	}
 }
 
-void FFractureEditorModeToolkit::OverridesDetailsViewFinishedChangingProperties(const FPropertyChangedEvent& InPropertyChangedEvent)
+FReply FFractureEditorModeToolkit::OnDataflowOverridesUpdateAsset()
 {
 	// get the displayed asset from the settings
 	UFractureSettings* FractureSettings = GetMutableDefault<UFractureSettings>();
-	if (FractureSettings->RestCollection.IsValid())
+	if (FractureSettings && FractureSettings->RestCollection.IsValid())
 	{
 		FractureSettings->RestCollection->GetDataflowInstance().UpdateOwnerAsset();
 	}
+	return FReply::Handled();
+}
+
+bool FFractureEditorModeToolkit::CanDataflowOverridesUpdateAsset() const
+{
+	return true;
+}
+
+bool FFractureEditorModeToolkit::GetDataflowOverridesUpdateAssetEnabled() const
+{
+	UFractureSettings* FractureSettings = GetMutableDefault<UFractureSettings>();
+	if (FractureSettings && FractureSettings->RestCollection.IsValid())
+	{
+		if (FractureSettings->RestCollection->GetDataflowInstance().GetDataflowAsset())
+		{
+			return true;
+		}
+	}
+	return false;
 }
 
 void FFractureEditorModeToolkit::Init(const TSharedPtr<IToolkitHost>& InitToolkitHost, TWeakObjectPtr<UEdMode> InOwningMode)
@@ -348,11 +365,12 @@ void FFractureEditorModeToolkit::Init(const TSharedPtr<IToolkitHost>& InitToolki
 	DetailsViewArgs.bShowAnimatedPropertiesOption = false;
 
 	DetailsView = EditModule.CreateDetailView(DetailsViewArgs);
+	FractureSettingsDetailsView = EditModule.CreateDetailView(DetailsViewArgs);
 	EditModule.RegisterCustomClassLayout("FractureSettings", FOnGetDetailCustomizationInstance::CreateStatic(&FFractureViewSettingsCustomization::MakeInstance, this));
 
 	TArray<UObject*> Settings;
 	Settings.Add(GetMutableDefault<UFractureSettings>());
-	DetailsView->SetObjects(Settings);
+	FractureSettingsDetailsView->SetObjects(Settings);
 
 	HistogramDetailsView = EditModule.CreateDetailView(DetailsViewArgs);
 	EditModule.RegisterCustomClassLayout("HistogramSettings", FOnGetDetailCustomizationInstance::CreateStatic(&FHistogramSettingsCustomization::MakeInstance, this));
@@ -399,31 +417,64 @@ void FFractureEditorModeToolkit::Init(const TSharedPtr<IToolkitHost>& InitToolki
 		.HeaderContent()
 		[
 			SNew(SHorizontalBox)
-			+ SHorizontalBox::Slot().HAlign(HAlign_Left).VAlign(VAlign_Center).FillWidth(2.f)
-			[
-				SNew(STextBlock)
-				.Text(LOCTEXT("FractureSettingsPanelHeader", "Fracture Mode Quick Settings"))
-			]
+				+ SHorizontalBox::Slot().HAlign(HAlign_Left).VAlign(VAlign_Center).FillWidth(2.f)
+				[
+					SNew(STextBlock)
+						.Text(LOCTEXT("FractureSettingsPanelHeader", "Fracture Mode Quick Settings"))
+				]
 
-			+ SHorizontalBox::Slot().HAlign(HAlign_Right).VAlign(VAlign_Center).AutoWidth()
-			[
-				SNew(SComboButton)
-				.HasDownArrow(false)
-				.MenuPlacement(EMenuPlacement::MenuPlacement_MenuRight)
-				.ComboButtonStyle(FAppStyle::Get(), "SimpleComboButton")
-				.OnGetMenuContent(FOnGetContent::CreateLambda([this]()
-				{
-					return MakeMenu_FractureModeConfigSettings();
-				}))
-				.ContentPadding(FMargin(3.0f, 1.0f))
-				.ButtonContent()
+				+ SHorizontalBox::Slot().HAlign(HAlign_Right).VAlign(VAlign_Center).AutoWidth()
 				[
-					SNew(SImage)
-					.Image(FFractureEditorStyle::Get().GetBrush("FractureEditor.DefaultSettings"))
-					.ColorAndOpacity(FSlateColor::UseForeground())
+					SNew(SComboButton)
+						.HasDownArrow(false)
+						.MenuPlacement(EMenuPlacement::MenuPlacement_MenuRight)
+						.ComboButtonStyle(FAppStyle::Get(), "SimpleComboButton")
+						.OnGetMenuContent(FOnGetContent::CreateLambda([this]()
+							{
+								return MakeMenu_FractureModeConfigSettings();
+							}))
+						.ContentPadding(FMargin(3.0f, 1.0f))
+								.ButtonContent()
+								[
+									SNew(SImage)
+										.Image(FFractureEditorStyle::Get().GetBrush("FractureEditor.DefaultSettings"))
+										.ColorAndOpacity(FSlateColor::UseForeground())
+								]
 				]
+
+		];
+
+
+	TSharedPtr<SVerticalBox> DataflowOverrideWidget =
+		SNew(SVerticalBox)
+		+ SVerticalBox::Slot()
+		.AutoHeight()
+		[
+			SNew(SScrollBox)
+			+ SScrollBox::Slot()
+			[
+				OverridesDetailsView->GetWidget().ToSharedRef()
+			]
+		]
+		+ SVerticalBox::Slot()
+		.AutoHeight()
+		[
+			SNew(SHorizontalBox)
+			+ SHorizontalBox::Slot()
+			[
+				SNew(SSpacer)
 			]
 
+			+ SHorizontalBox::Slot()
+			.Padding(4.0)
+			.AutoWidth()
+			[
+				SNew(SPrimaryButton)
+					.OnClicked(this, &FFractureEditorModeToolkit::OnDataflowOverridesUpdateAsset)
+					.IsEnabled(this, &FFractureEditorModeToolkit::CanDataflowOverridesUpdateAsset)
+					.Text(LOCTEXT("DataflowOverrideUpdateAssetButton", "Update Asset From Dataflow"))
+					.IsEnabled_Lambda([this]() -> bool { return GetDataflowOverridesUpdateAssetEnabled(); })
+			]
 		];
 
 	const TSharedPtr<SVerticalBox> ToolkitWidgetVBox = SNew(SVerticalBox);
@@ -453,16 +504,21 @@ void FFractureEditorModeToolkit::Init(const TSharedPtr<IToolkitHost>& InitToolki
 					SNew(SScrollBox)
 					+SScrollBox::Slot()
 					[
-						DetailsView.ToSharedRef()
+						FractureSettingsDetailsView.ToSharedRef()
 					]
 				]
 				+ SVerticalBox::Slot()
+				.AutoHeight()
+				[
+					DataflowOverrideWidget.ToSharedRef()
+				]
+				+ SVerticalBox::Slot()
 				.FillHeight(1.0)
 				[
 					SNew(SScrollBox)
 					+ SScrollBox::Slot()
 					[
-						OverridesDetailsView->GetWidget().ToSharedRef()
+						DetailsView.ToSharedRef()
 					]
 				]
 				+SVerticalBox::Slot()
@@ -1716,22 +1772,20 @@ void FFractureEditorModeToolkit::SetActiveTool(UFractureModalTool* InActiveTool)
 
 	ActiveTool = InActiveTool;
 
-	TArray<UObject*> Settings;
-	Settings.Add(GetMutableDefault<UFractureSettings>());
-
+	TArray<UObject*> ToolSettings;
 	if (ActiveTool != nullptr)
 	{
 		ActiveTool->OnPropertyModifiedDirectlyByTool.AddSP(this, &FFractureEditorModeToolkit::InvalidateCachedDetailPanelState);
 
 		ActiveTool->Setup(StaticCastSharedRef<FFractureEditorModeToolkit>(AsShared()));
 
-		Settings.Append(ActiveTool->GetSettingsObjects());
+		ToolSettings.Append(ActiveTool->GetSettingsObjects());
 
 		ActiveTool->SelectedBonesChanged();
 		ActiveTool->FractureContextChanged();
 	}
 
-	DetailsView->SetObjects(Settings);
+	DetailsView->SetObjects(ToolSettings);
 }
 
 void FFractureEditorModeToolkit::InvalidateCachedDetailPanelState(UObject* ChangedObject)
diff --git a/Engine/Plugins/Experimental/ChaosEditor/Source/FractureEditor/Public/FractureEditorModeToolkit.h b/Engine/Plugins/Experimental/ChaosEditor/Source/FractureEditor/Public/FractureEditorModeToolkit.h
--- a/Engine/Plugins/Experimental/ChaosEditor/Source/FractureEditor/Public/FractureEditorModeToolkit.h
+++ b/Engine/Plugins/Experimental/ChaosEditor/Source/FractureEditor/Public/FractureEditorModeToolkit.h
@@ -246,8 +246,10 @@ private:
 
 	void CreateVariableOverrideDetailView();
 	void RefreshVariableOverrideDetailView(const UGeometryCollection* RestCollection);
-	void OverridesDetailsViewFinishedChangingProperties(const FPropertyChangedEvent& InPropertyChangedEvent);
 
+	FReply OnDataflowOverridesUpdateAsset();
+	bool CanDataflowOverridesUpdateAsset() const;
+	bool GetDataflowOverridesUpdateAssetEnabled() const;
 private:
 	TObjectPtr<UFractureModalTool> ActiveTool;
 
@@ -257,6 +259,7 @@ private:
 	FDelegateHandle ProjectSettingsModifiedHandle;
 
 	TSharedPtr<IDetailsView> DetailsView;
+	TSharedPtr<IDetailsView> FractureSettingsDetailsView;
 	TSharedPtr<IStructureDetailsView> OverridesDetailsView;
 	TSharedPtr<IDetailsView> HistogramDetailsView;
 	TSharedPtr<IDetailsView> OutlinerDetailsView;
