Fracture tool : add physical material option to creation of a geometry collection and auto generate root proxies from geometry sources

diff --git a/Engine/Plugins/Experimental/ChaosEditor/Source/FractureEditor/FractureEditor.Build.cs b/Engine/Plugins/Experimental/ChaosEditor/Source/FractureEditor/FractureEditor.Build.cs
--- a/Engine/Plugins/Experimental/ChaosEditor/Source/FractureEditor/FractureEditor.Build.cs
+++ b/Engine/Plugins/Experimental/ChaosEditor/Source/FractureEditor/FractureEditor.Build.cs
@@ -51,7 +51,8 @@ public class FractureEditor : ModuleRules
 				"FractureEngine",
 				"ContentBrowserData",
 				"GeometryFramework",
-				"DataflowCore"
+				"DataflowCore",
+				"PhysicsCore"
 
 				// ... add private dependencies that you statically link with here ...	
 			}
diff --git a/Engine/Plugins/Experimental/ChaosEditor/Source/FractureEditor/Private/FractureToolGenerators.cpp b/Engine/Plugins/Experimental/ChaosEditor/Source/FractureEditor/Private/FractureToolGenerators.cpp
--- a/Engine/Plugins/Experimental/ChaosEditor/Source/FractureEditor/Private/FractureToolGenerators.cpp
+++ b/Engine/Plugins/Experimental/ChaosEditor/Source/FractureEditor/Private/FractureToolGenerators.cpp
@@ -34,6 +34,8 @@
 #include "Input/Reply.h"
 #include "Widgets/SCompoundWidget.h"
 
+#include "PhysicalMaterials/PhysicalMaterial.h"
+#include "GeometryCollection/GeometryCollectionRootProxyRenderer.h"
 #include "Modules/ModuleManager.h"
 #include "Widgets/SBoxPanel.h"
 #include "Widgets/Text/STextBlock.h"
@@ -78,7 +80,7 @@ namespace
 //////////////////////////////////////////////////////////////////////////
 // SCreateGeometryCollectionFromObject is adapted from SCreateAssetFromActor 
 
-DECLARE_DELEGATE_FourParams(FOnPathChosen, const FString&, bool, bool, UDataflow*);
+DECLARE_DELEGATE_OneParam(FOnPathChosen, const UFractureToolGenerateAsset::FCreationParameters& Params);
 
 class SCreateGeometryCollectionFromObject : public SCompoundWidget
 {
@@ -91,6 +93,7 @@ public:
 		, _bSplitIslands(false)
 		, _bReimportToMeshOutput(false)
 		, _DataflowAsset(nullptr)
+		, _PhysicalMaterial(nullptr)
 	{}
 
 	/** The default suffix to use for the asset filename */
@@ -112,7 +115,9 @@ public:
 
 	SLATE_ARGUMENT(bool, bSplitIslands)
 	SLATE_ARGUMENT(bool, bReimportToMeshOutput)
+	SLATE_ARGUMENT(bool, bUseRootProxies)
 	SLATE_ARGUMENT(UDataflow*, DataflowAsset)
+	SLATE_ARGUMENT(UPhysicalMaterial*, PhysicalMaterial)
 
 	/** Action to perform when create clicked */
 	SLATE_EVENT(FOnPathChosen, OnCreateAssetAction)
@@ -146,6 +151,23 @@ private:
 		}
 	}
 
+	FString GetPhysicalMaterialPath() const
+	{
+		if (PhysicalMaterial)
+		{
+			return PhysicalMaterial->GetPathName();
+		}
+		return FString();
+	}
+	void OnPhysicalMaterialChanged(const FAssetData& InAssetData)
+	{
+		if (UPhysicalMaterial* NewPhysicalMaterial = Cast<UPhysicalMaterial>(InAssetData.GetAsset()))
+		{
+			PhysicalMaterial = NewPhysicalMaterial;
+		}
+	}
+	
+
 	/** Destroys the window when the operation is cancelled. */
 	FReply OnCancelCreateAssetFromActor();
 
@@ -189,9 +211,15 @@ private:
 	/** Designate input faces with odd-numbered material indices as 'internal' materials. Only use this option to preserve internal faces on meshes created with the Fracture Mode 'ToMesh' tool. */
 	bool bFromToMeshTool = false;
 
+	/** Whether to use the meshes used to create the geometry collection as root proxies */
+	bool bUseRootProxies = true;
+
 	/** Dataflow asset to assign */
 	TObjectPtr<UDataflow> DataflowAsset = nullptr;
 
+	/** Physical Material to assign */
+	TObjectPtr<UPhysicalMaterial> PhysicalMaterial = nullptr;
+
 	/** Filename textbox widget */
 	TSharedPtr<SEditableTextBox> FileNameWidget;
 
@@ -216,7 +244,9 @@ void SCreateGeometryCollectionFromObject::Construct(const FArguments& InArgs, TS
 	OnCreateAssetAction = InArgs._OnCreateAssetAction;
 	bSplitIslands = InArgs._bSplitIslands;
 	bFromToMeshTool = InArgs._bReimportToMeshOutput;
+	bUseRootProxies = InArgs._bUseRootProxies;
 	DataflowAsset = InArgs._DataflowAsset;
+	PhysicalMaterial = InArgs._PhysicalMaterial;
 
 
 	if (InArgs._AssetPath.IsEmpty())
@@ -390,6 +420,38 @@ void SCreateGeometryCollectionFromObject::Construct(const FArguments& InArgs, TS
 					]
 				]
 				+ SVerticalBox::Slot()
+				[
+					SNew(SHorizontalBox)
+					+ SHorizontalBox::Slot()
+					.Padding(0.0f, 0.0f, 8.0f, 0.0f)
+					.FillWidth(0.5f)
+					.VAlign(VAlign_Center)
+					[
+						SNew(STextBlock)
+						.Text(LOCTEXT("UseRootProxiesLabel", "Use Root Proxies"))
+						.ToolTipText(LOCTEXT("UseRootProxiesToolTip", "If checked, the meshes used to create the geometry collection will be used as root proxies for rendering the unbroken state."))
+					]
+
+					+ SHorizontalBox::Slot()
+					.FillWidth(1.0f)
+					[
+						SNew(SCheckBox)
+						.IsChecked_Lambda([this]()->ECheckBoxState
+							{
+								return bUseRootProxies ? ECheckBoxState::Checked : ECheckBoxState::Unchecked;
+							})
+						.OnCheckStateChanged_Lambda([this](ECheckBoxState NewState)
+						{
+							if (NewState == ECheckBoxState::Undetermined)
+							{
+								return;
+							}
+							bUseRootProxies = NewState == ECheckBoxState::Checked;
+						})
+						.ToolTipText(LOCTEXT("UseRootProxiesToolTip", "If checked, the meshes used to create the geometry collection will be used as root proxies for rendering the unbroken state."))
+					]
+				]
+				+ SVerticalBox::Slot()
 				[
 					SNew(SHorizontalBox)
 					+ SHorizontalBox::Slot()
@@ -406,14 +468,41 @@ void SCreateGeometryCollectionFromObject::Construct(const FArguments& InArgs, TS
 					.FillWidth(1.0f)
 					[
 						SNew(SObjectPropertyEntryBox)
-							.ObjectPath(this, &SCreateGeometryCollectionFromObject::GetDataflowAssetPath)
-							.OnObjectChanged(this, &SCreateGeometryCollectionFromObject::OnDataflowAssetChanged)
-							.AllowedClass(UDataflow::StaticClass())
-							.DisplayUseSelected(true)
-							.DisplayBrowse(true)
-							.DisplayThumbnail(true)
-							.AllowCreate(true)
-							.AllowClear(true)
+						.ObjectPath(this, &SCreateGeometryCollectionFromObject::GetDataflowAssetPath)
+						.OnObjectChanged(this, &SCreateGeometryCollectionFromObject::OnDataflowAssetChanged)
+						.AllowedClass(UDataflow::StaticClass())
+						.DisplayUseSelected(true)
+						.DisplayBrowse(true)
+						.DisplayThumbnail(true)
+						.AllowCreate(true)
+						.AllowClear(true)
+					]
+				]
+				+ SVerticalBox::Slot()
+				[
+					SNew(SHorizontalBox)
+					+ SHorizontalBox::Slot()
+					.Padding(0.0f, 0.0f, 8.0f, 0.0f)
+					.FillWidth(0.5f)
+					.VAlign(VAlign_Center)
+					[
+						SNew(STextBlock)
+						.Text(LOCTEXT("PhysicalMaterialAssetLabel", "PhysicalMaterial"))
+						.ToolTipText(LOCTEXT("PhysicalMaterialAssetToolTip", "Optionally assign a Physical material to the new Geometry Collection."))
+					]
+
+					+ SHorizontalBox::Slot()
+					.FillWidth(1.0f)
+					[
+						SNew(SObjectPropertyEntryBox)
+						.ObjectPath(this, &SCreateGeometryCollectionFromObject::GetPhysicalMaterialPath)
+						.OnObjectChanged(this, &SCreateGeometryCollectionFromObject::OnPhysicalMaterialChanged)
+						.AllowedClass(UPhysicalMaterial::StaticClass())
+						.DisplayUseSelected(true)
+						.DisplayBrowse(true)
+						.DisplayThumbnail(true)
+						.AllowCreate(true)
+						.AllowClear(true)
 					]
 				]
 			]
@@ -462,7 +551,16 @@ void SCreateGeometryCollectionFromObject::RequestDestroyParentWindow()
 FReply SCreateGeometryCollectionFromObject::OnCreateAssetFromActorClicked()
 {
 	RequestDestroyParentWindow();
-	OnCreateAssetAction.ExecuteIfBound(AssetPath / FileNameWidget->GetText().ToString(), bFromToMeshTool, bSplitIslands, DataflowAsset);
+	UFractureToolGenerateAsset::FCreationParameters Params
+	{
+		.AssetPath = AssetPath / FileNameWidget->GetText().ToString(),
+		.bFromToMeshTool = bFromToMeshTool,
+		.bSplitIslands = bSplitIslands,
+		.bUseRootProxies = bUseRootProxies,
+		.DataflowAsset = DataflowAsset,
+		.PhysicalMaterial = PhysicalMaterial,
+	};
+	OnCreateAssetAction.ExecuteIfBound(Params);
 	return FReply::Handled();
 }
 
@@ -656,8 +754,10 @@ void UFractureToolGenerateAsset::OpenGenerateAssetDialog(TArray<AActor*>& Actors
 		.CreateButtonText(LOCTEXT("CreateGeometryCollection_ButtonLabel", "Create Geometry Collection"))
 		.AssetPath(AssetPath)
 		.DataflowAsset(LastDataflowAsset)
+		.PhysicalMaterial(LastPhysicalMaterial)
 		.bReimportToMeshOutput(bLastFromToMesh)
 		.bSplitIslands(bLastSplitComponents)
+		.bUseRootProxies(bLastUseRootProxies)
 		.OnCreateAssetAction(FOnPathChosen::CreateUObject(this, &UFractureToolGenerateAsset::OnGenerateAssetPathChosen, Actors))
 	);
 
@@ -673,17 +773,19 @@ void UFractureToolGenerateAsset::OpenGenerateAssetDialog(TArray<AActor*>& Actors
 
 }
 
-void UFractureToolGenerateAsset::OnGenerateAssetPathChosen(const FString& InAssetPath, bool bFromToMeshTool, bool bSplitComponents, UDataflow* Dataflow, TArray<AActor*> Actors)
+void UFractureToolGenerateAsset::OnGenerateAssetPathChosen(const UFractureToolGenerateAsset::FCreationParameters& Params, TArray<AActor*> Actors)
 {	
 	//Record the path and other settings
 	int32 LastSlash = INDEX_NONE;
-	if (InAssetPath.FindLastChar('/', LastSlash))
+	if (Params.AssetPath.FindLastChar('/', LastSlash))
 	{
-		AssetPath = InAssetPath.Left(LastSlash);
+		AssetPath = Params.AssetPath.Left(LastSlash);
 	}
-	bLastFromToMesh = bFromToMeshTool;
-	bLastSplitComponents = bSplitComponents;
-	LastDataflowAsset = Dataflow;
+	bLastFromToMesh = Params.bFromToMeshTool;
+	bLastSplitComponents = Params.bSplitIslands;
+	bLastUseRootProxies= Params.bUseRootProxies;
+	LastDataflowAsset = Params.DataflowAsset;
+	LastPhysicalMaterial = Params.PhysicalMaterial;
 
 	UGeometryCollectionComponent* GeometryCollectionComponent = nullptr;
 
@@ -693,7 +795,7 @@ void UFractureToolGenerateAsset::OnGenerateAssetPathChosen(const FString& InAsse
 
 		AGeometryCollectionActor* GeometryCollectionActor = nullptr;
 		
-		GeometryCollectionActor = ConvertActorsToGeometryCollection(InAssetPath, false/*bAddInternalMaterials*/, bSplitComponents, Dataflow, Actors, bFromToMeshTool);
+		GeometryCollectionActor = ConvertActorsToGeometryCollection(Params, false/*bAddInternalMaterials*/, Actors);
 		if (!GeometryCollectionActor)
 		{
 			return;
@@ -758,7 +860,7 @@ void UFractureToolGenerateAsset::OnGenerateAssetPathChosen(const FString& InAsse
 	}
 }
 
-AGeometryCollectionActor* UFractureToolGenerateAsset::ConvertActorsToGeometryCollection(const FString& InAssetPath, bool bAddInternalMaterials, bool bSplitComponents, UDataflow* Dataflow, TArray<AActor*>& Actors, bool bFromToMeshTool)
+AGeometryCollectionActor* UFractureToolGenerateAsset::ConvertActorsToGeometryCollection(const FCreationParameters& Params, bool bAddInternalMaterials, TArray<AActor*>& Actors)
 {
 	ensure(!bAddInternalMaterials); // we should not use the 'add internal materials' path anymore, as we move away from the odd-numbered-internal-material convention
 	ensure(Actors.Num() > 0);
@@ -818,7 +920,7 @@ AGeometryCollectionActor* UFractureToolGenerateAsset::ConvertActorsToGeometryCol
 	FScopedSlowTask SlowTask(SourceTriCount, LOCTEXT("FractureCreateNewGeometryCollection", "Creating new Geometry Collection"));
 	SlowTask.MakeDialog();
 
-	AGeometryCollectionActor* NewActor = CreateNewGeometryActor(InAssetPath, FTransform(), true);
+	AGeometryCollectionActor* NewActor = CreateNewGeometryActor(Params.AssetPath, FTransform(), true);
 
 	FGeometryCollectionEdit GeometryCollectionEdit = NewActor->GetGeometryCollectionComponent()->EditRestCollection(GeometryCollection::EEditUpdate::RestPhysicsDynamic);
 	UGeometryCollection* FracturedGeometryCollection = GeometryCollectionEdit.GetRestCollection();
@@ -851,11 +953,11 @@ AGeometryCollectionActor* UFractureToolGenerateAsset::ConvertActorsToGeometryCol
 				// Record the contributing source on the asset.
 				FSoftObjectPath SourceSoftObjectPath(ComponentStaticMesh);
 				decltype(FGeometryCollectionSource::SourceMaterial) SourceMaterials(StaticMeshComponent->GetMaterials());
-				bool bSetInternalFromMaterialIndex = bFromToMeshTool;
+				bool bSetInternalFromMaterialIndex = Params.bFromToMeshTool;
 				
-				FracturedGeometryCollection->GeometrySource.Emplace(SourceSoftObjectPath, ComponentTransform, SourceMaterials, bSplitComponents, bSetInternalFromMaterialIndex);
+				FracturedGeometryCollection->GeometrySource.Emplace(SourceSoftObjectPath, ComponentTransform, SourceMaterials, Params.bSplitIslands, bSetInternalFromMaterialIndex);
 				
-				FGeometryCollectionEngineConversion::AppendStaticMesh(ComponentStaticMesh, SourceMaterials, ComponentTransform, FracturedGeometryCollection, false/*bReindexMaterials*/, bAddInternalMaterials, bSplitComponents, bSetInternalFromMaterialIndex);
+				FGeometryCollectionEngineConversion::AppendStaticMesh(ComponentStaticMesh, SourceMaterials, ComponentTransform, FracturedGeometryCollection, false/*bReindexMaterials*/, bAddInternalMaterials, Params.bSplitIslands, bSetInternalFromMaterialIndex);
 			}
 		}
 
@@ -900,11 +1002,33 @@ AGeometryCollectionActor* UFractureToolGenerateAsset::ConvertActorsToGeometryCol
 		}
 	}
 
+	if (Params.PhysicalMaterial)
+	{
+		FracturedGeometryCollection->PhysicsMaterial = Params.PhysicalMaterial;
+	}
+
+	if (Params.bUseRootProxies)
+	{
+		FracturedGeometryCollection->CustomRendererType = UGeometryCollectionRootProxyRenderer::StaticClass();
+		FracturedGeometryCollection->RootProxyData.ProxyMeshes.Reserve(FracturedGeometryCollection->GeometrySource.Num());
+		FracturedGeometryCollection->RootProxyData.MeshTransforms.Reserve(FracturedGeometryCollection->GeometrySource.Num());
+
+		for (const FGeometryCollectionSource& Source : FracturedGeometryCollection->GeometrySource)
+		{
+			UObject* SourceMesh = Source.SourceGeometryObject.TryLoad();
+			if (UStaticMesh* SourceStaticMesh = Cast<UStaticMesh>(SourceMesh))
+			{
+				FracturedGeometryCollection->RootProxyData.ProxyMeshes.Add(SourceStaticMesh);
+				FracturedGeometryCollection->RootProxyData.MeshTransforms.Add(FTransform3f(Source.LocalTransform));
+			}
+		}
+	}
+
 	FracturedGeometryCollection->InitializeMaterials(bAddInternalMaterials);
 
-	if (Dataflow)
+	if (Params.DataflowAsset)
 	{
-		FracturedGeometryCollection->DataflowInstance.SetDataflowAsset(Dataflow);
+		FracturedGeometryCollection->DataflowInstance.SetDataflowAsset(Params.DataflowAsset);
 	}
 
 	AddSingleRootNodeIfRequired(FracturedGeometryCollection);
diff --git a/Engine/Plugins/Experimental/ChaosEditor/Source/FractureEditor/Private/FractureToolGenerators.h b/Engine/Plugins/Experimental/ChaosEditor/Source/FractureEditor/Private/FractureToolGenerators.h
--- a/Engine/Plugins/Experimental/ChaosEditor/Source/FractureEditor/Private/FractureToolGenerators.h
+++ b/Engine/Plugins/Experimental/ChaosEditor/Source/FractureEditor/Private/FractureToolGenerators.h
@@ -16,6 +16,16 @@ class UFractureToolGenerateAsset : public UFractureActionTool
 public:
 	GENERATED_BODY()
 
+	struct FCreationParameters
+	{
+		FString AssetPath;
+		bool bFromToMeshTool = false;
+		bool bSplitIslands = false;
+		bool bUseRootProxies = true;
+		TObjectPtr<UDataflow> DataflowAsset;
+		TObjectPtr<UPhysicalMaterial> PhysicalMaterial;
+	};
+
 	UFractureToolGenerateAsset(const FObjectInitializer& ObjInit) 
 		: Super(ObjInit)
 		, AssetPath(FString(TEXT("/Game")))
@@ -31,8 +41,8 @@ public:
 
 private:
 	void OpenGenerateAssetDialog(TArray<AActor*>& Actors);
-	void OnGenerateAssetPathChosen(const FString& InAssetPath, bool bFromToMeshTool, bool bSplitComponents, UDataflow* Dataflow, TArray<AActor*> Actors);
-	AGeometryCollectionActor* ConvertActorsToGeometryCollection(const FString& InAssetPath, bool bAddInternalMaterials, bool bSplitComponents, UDataflow* Dataflow, TArray<AActor*>& Actors, bool bFromToMeshTool);
+	void OnGenerateAssetPathChosen(const FCreationParameters& Params, TArray<AActor*> Actors);
+	AGeometryCollectionActor* ConvertActorsToGeometryCollection(const FCreationParameters& Params, bool bAddInternalMaterials, TArray<AActor*>& Actors);
 	class AGeometryCollectionActor* CreateNewGeometryActor(const FString& Name, const FTransform& Transform, bool AddMaterials /*= false*/);
 	static AActor* AddActor(ULevel* InLevel, UClass* Class);
 
@@ -47,8 +57,12 @@ private:
 	UPROPERTY()
 	TObjectPtr<UDataflow> LastDataflowAsset = nullptr;
 
+	UPROPERTY()
+	TObjectPtr<UPhysicalMaterial> LastPhysicalMaterial = nullptr;
+
 	bool bLastSplitComponents = false;
 	bool bLastFromToMesh = false;
+	bool bLastUseRootProxies = true;
 };
 
 
