Fracture editor : add a box around selection so that nanite enabled geometry collection selected bones can be visualized

diff --git a/Engine/Plugins/Experimental/ChaosEditor/Source/FractureEditor/Private/FractureEditorMode.cpp b/Engine/Plugins/Experimental/ChaosEditor/Source/FractureEditor/Private/FractureEditorMode.cpp
--- a/Engine/Plugins/Experimental/ChaosEditor/Source/FractureEditor/Private/FractureEditorMode.cpp
+++ b/Engine/Plugins/Experimental/ChaosEditor/Source/FractureEditor/Private/FractureEditorMode.cpp
@@ -223,7 +223,7 @@ void UFractureEditorMode::Render(const FSceneView* View, FViewport* Viewport, FP
 		FractureTool->Render(View, Viewport, PDI);
 	}
 
-
+	RenderSelectionAsBoxes(View, Viewport, PDI);
 }
 
 void UFractureEditorMode::DrawHUD(FEditorViewportClient* ViewportClient, FViewport* Viewport, const FSceneView* View, FCanvas* Canvas)
@@ -237,6 +237,64 @@ void UFractureEditorMode::DrawHUD(FEditorViewportClient* ViewportClient, FViewpo
 	}
 }
 
+void UFractureEditorMode::RenderSelectionAsBoxes(const FSceneView* View, FViewport* Viewport, FPrimitiveDrawInterface* PDI)
+{
+	for (TWeakObjectPtr<UGeometryCollectionComponent> SelectedCompWeakPtr : SelectedGeometryComponents)
+	{
+		if (UGeometryCollectionComponent* SelectedComp = SelectedCompWeakPtr.Get())
+		{
+			const FTransform& WorldTransform = SelectedComp->GetComponentToWorld();
+
+			const TManagedArray<FVector3f>* ExplodedVectors = nullptr;
+			if (SelectedComp->GetRestCollection())
+			{
+				if (TSharedPtr<const FGeometryCollection> Collection = SelectedComp->GetRestCollection()->GetGeometryCollection())
+				{
+					ExplodedVectors = Collection->FindAttribute<FVector3f>("ExplodedVector", FGeometryCollection::TransformGroup);
+				}
+			}
+
+			for (int32 SelectedBone : SelectedComp->GetSelectedBones())
+			{
+				float BoneMass = 0;
+				FBox BoneExtent;
+				SelectedComp->GetMassAndExtents(SelectedBone, BoneMass, BoneExtent);
+
+				const FVector3f Offset = ExplodedVectors ? (*ExplodedVectors)[SelectedBone] : FVector3f::ZeroVector;
+				const FVector BoxCenter = BoneExtent.GetCenter() + FVector(Offset);
+				const FVector BoxHalfSize = BoneExtent.GetExtent();
+				const FVector Corners[] =
+				{
+					WorldTransform.TransformPosition(BoxCenter + FVector(-BoxHalfSize.X, -BoxHalfSize.Y, -BoxHalfSize.Z)),
+					WorldTransform.TransformPosition(BoxCenter + FVector(+BoxHalfSize.X, -BoxHalfSize.Y, -BoxHalfSize.Z)),
+					WorldTransform.TransformPosition(BoxCenter + FVector(-BoxHalfSize.X, +BoxHalfSize.Y, -BoxHalfSize.Z)),
+					WorldTransform.TransformPosition(BoxCenter + FVector(+BoxHalfSize.X, +BoxHalfSize.Y, -BoxHalfSize.Z)),
+					WorldTransform.TransformPosition(BoxCenter + FVector(-BoxHalfSize.X, -BoxHalfSize.Y, +BoxHalfSize.Z)),
+					WorldTransform.TransformPosition(BoxCenter + FVector(+BoxHalfSize.X, -BoxHalfSize.Y, +BoxHalfSize.Z)),
+					WorldTransform.TransformPosition(BoxCenter + FVector(-BoxHalfSize.X, +BoxHalfSize.Y, +BoxHalfSize.Z)),
+					WorldTransform.TransformPosition(BoxCenter + FVector(+BoxHalfSize.X, +BoxHalfSize.Y, +BoxHalfSize.Z)),
+				};
+				// bottom 
+				PDI->DrawLine(Corners[0], Corners[1], FLinearColor::Yellow, 1.f, 0.0f);
+				PDI->DrawLine(Corners[1], Corners[3], FLinearColor::Yellow, 1.f, 0.0f);
+				PDI->DrawLine(Corners[3], Corners[2], FLinearColor::Yellow, 1.f, 0.0f);
+				PDI->DrawLine(Corners[2], Corners[0], FLinearColor::Yellow, 1.f, 0.0f);
+
+				// top 
+				PDI->DrawLine(Corners[4], Corners[5], FLinearColor::Yellow, 1.f, 0.0f);
+				PDI->DrawLine(Corners[5], Corners[7], FLinearColor::Yellow, 1.f, 0.0f);
+				PDI->DrawLine(Corners[7], Corners[6], FLinearColor::Yellow, 1.f, 0.0f);
+				PDI->DrawLine(Corners[6], Corners[4], FLinearColor::Yellow, 1.f, 0.0f);
+
+				// vertical edges
+				PDI->DrawLine(Corners[0], Corners[4], FLinearColor::Yellow, 1.f, 0.0f);
+				PDI->DrawLine(Corners[1], Corners[5], FLinearColor::Yellow, 1.f, 0.0f);
+				PDI->DrawLine(Corners[2], Corners[6], FLinearColor::Yellow, 1.f, 0.0f);
+				PDI->DrawLine(Corners[3], Corners[7], FLinearColor::Yellow, 1.f, 0.0f);
+			}
+		}
+	}
+}
 
 void UFractureEditorMode::CreateToolkit()
 {
diff --git a/Engine/Plugins/Experimental/ChaosEditor/Source/FractureEditor/Public/FractureEditorMode.h b/Engine/Plugins/Experimental/ChaosEditor/Source/FractureEditor/Public/FractureEditorMode.h
--- a/Engine/Plugins/Experimental/ChaosEditor/Source/FractureEditor/Public/FractureEditorMode.h
+++ b/Engine/Plugins/Experimental/ChaosEditor/Source/FractureEditor/Public/FractureEditorMode.h
@@ -84,6 +84,8 @@ private:
 
 	void ConfigureRealTimeViewportsOverride(bool bEnable);
 
+	void RenderSelectionAsBoxes(const FSceneView* View, FViewport* Viewport, FPrimitiveDrawInterface* PDI);
+
 private:
 	/** This selection set is updated from actor selection changed event.  We change state on components as they are selected so we have to maintain or own list **/
 	TArray<TWeakObjectPtr<UGeometryCollectionComponent>> SelectedGeometryComponents;
