Fracture editor : Fix crashes when using fracture tools on a geometry collection component with no rest collection
- made sure to check that the geometry collection component and tool contexts are valid before executing the tools
- tested on the whole suite of tools in the fracture editor

diff --git a/Engine/Plugins/Experimental/ChaosEditor/Source/FractureEditor/Private/FractureToolAutoCluster.cpp b/Engine/Plugins/Experimental/ChaosEditor/Source/FractureEditor/Private/FractureToolAutoCluster.cpp
--- a/Engine/Plugins/Experimental/ChaosEditor/Source/FractureEditor/Private/FractureToolAutoCluster.cpp
+++ b/Engine/Plugins/Experimental/ChaosEditor/Source/FractureEditor/Private/FractureToolAutoCluster.cpp
@@ -127,51 +127,54 @@ void UFractureToolAutoCluster::FractureContextChanged()
 
 	for (FFractureToolContext& Context : Contexts)
 	{
-		Context.ConvertSelectionToClusterNodes();
-
-		if (const FGeometryCollection* GeometryCollection = Context.GetGeometryCollection().Get())
+		if (Context.GetGeometryCollection())
 		{
-			int CollectionIdx = VisualizedCollections.Emplace(Context.GetGeometryCollectionComponent());
-			FTransform ContextTransform = Context.GetTransform();
+			Context.ConvertSelectionToClusterNodes();
 
-			for (const int32 ClusterIndex : Context.GetSelection())
+			if (const FGeometryCollection* GeometryCollection = Context.GetGeometryCollection().Get())
 			{
-				GridPointsMappings.AddMapping(CollectionIdx, ClusterIndex, ShowGridPoints.Num());
-				FBox Bounds;
-				TArray<FVector> LocalGridPoints = FFractureEngineClustering::GenerateGridSites(
-					*GeometryCollection, ClusterIndex,
-					AutoClusterSettings->ClusterGridWidth,
-					AutoClusterSettings->ClusterGridDepth,
-					AutoClusterSettings->ClusterGridHeight,
-					&Bounds);
-				for (FVector LocalPt : LocalGridPoints)
-				{
-					ShowGridPoints.Add(ContextTransform.TransformPosition(LocalPt));
-				}
-				GridLinesMappings.AddMapping(CollectionIdx, ClusterIndex, ShowGridLines.Num());
-				FIntVector3 Dims(
-					AutoClusterSettings->ClusterGridWidth,
-					AutoClusterSettings->ClusterGridDepth,
-					AutoClusterSettings->ClusterGridHeight);
-				FVector InvDims(1.0 / double(Dims.X), 1.0 / double(Dims.Y), 1.0 / double(Dims.Z));
-				FVector Diag = Bounds.Max - Bounds.Min;
-				for (int32 Dim = 0; Dim < 3; ++Dim)
+				int CollectionIdx = VisualizedCollections.Emplace(Context.GetGeometryCollectionComponent());
+				FTransform ContextTransform = Context.GetTransform();
+
+				for (const int32 ClusterIndex : Context.GetSelection())
 				{
-					FIntVector3 DMap(Dim, (Dim + 1) % 3, (Dim + 2) % 3);
-					FVector Start, End;
-					Start[DMap.Z] = Bounds.Min[DMap.Z];
-					End[DMap.Z] = Bounds.Max[DMap.Z];
-					for (int32 W = 0; W <= Dims[DMap.X]; ++W)
+					GridPointsMappings.AddMapping(CollectionIdx, ClusterIndex, ShowGridPoints.Num());
+					FBox Bounds;
+					TArray<FVector> LocalGridPoints = FFractureEngineClustering::GenerateGridSites(
+						*GeometryCollection, ClusterIndex,
+						AutoClusterSettings->ClusterGridWidth,
+						AutoClusterSettings->ClusterGridDepth,
+						AutoClusterSettings->ClusterGridHeight,
+						&Bounds);
+					for (FVector LocalPt : LocalGridPoints)
 					{
-						double AlongW = W * InvDims[DMap.X];
-						Start[DMap.X] = Bounds.Min[DMap.X] + AlongW * Diag[DMap.X];
-						End[DMap.X] = Start[DMap.X];
-						for (int32 H = 0; H <= Dims[DMap.Y]; ++H)
+						ShowGridPoints.Add(ContextTransform.TransformPosition(LocalPt));
+					}
+					GridLinesMappings.AddMapping(CollectionIdx, ClusterIndex, ShowGridLines.Num());
+					FIntVector3 Dims(
+						AutoClusterSettings->ClusterGridWidth,
+						AutoClusterSettings->ClusterGridDepth,
+						AutoClusterSettings->ClusterGridHeight);
+					FVector InvDims(1.0 / double(Dims.X), 1.0 / double(Dims.Y), 1.0 / double(Dims.Z));
+					FVector Diag = Bounds.Max - Bounds.Min;
+					for (int32 Dim = 0; Dim < 3; ++Dim)
+					{
+						FIntVector3 DMap(Dim, (Dim + 1) % 3, (Dim + 2) % 3);
+						FVector Start, End;
+						Start[DMap.Z] = Bounds.Min[DMap.Z];
+						End[DMap.Z] = Bounds.Max[DMap.Z];
+						for (int32 W = 0; W <= Dims[DMap.X]; ++W)
 						{
-							double AlongH = H * InvDims[DMap.Y];
-							Start[DMap.Y] = Bounds.Min[DMap.Y] + AlongH * Diag[DMap.Y];
-							End[DMap.Y] = Start[DMap.Y];
-							ShowGridLines.Emplace(ContextTransform.TransformPosition(Start), ContextTransform.TransformPosition(End));
+							double AlongW = W * InvDims[DMap.X];
+							Start[DMap.X] = Bounds.Min[DMap.X] + AlongW * Diag[DMap.X];
+							End[DMap.X] = Start[DMap.X];
+							for (int32 H = 0; H <= Dims[DMap.Y]; ++H)
+							{
+								double AlongH = H * InvDims[DMap.Y];
+								Start[DMap.Y] = Bounds.Min[DMap.Y] + AlongH * Diag[DMap.Y];
+								End[DMap.Y] = Start[DMap.Y];
+								ShowGridLines.Emplace(ContextTransform.TransformPosition(Start), ContextTransform.TransformPosition(End));
+							}
 						}
 					}
 				}
@@ -196,41 +199,44 @@ void UFractureToolAutoCluster::Execute(TWeakPtr<FFractureEditorModeToolkit> InTo
 		
 		for (FFractureToolContext& Context : Contexts)
 		{
-			FGeometryCollectionEdit Edit(Context.GetGeometryCollectionComponent(), GeometryCollection::EEditUpdate::RestPhysicsDynamic);
-
-			Context.ConvertSelectionToClusterNodes();
-			if (FGeometryCollection* GeometryCollection = Context.GetGeometryCollection().Get())
+			if (Context.GetGeometryCollection())
 			{
-				int32 StartTransformCount = GeometryCollection->Transform.Num();
+				FGeometryCollectionEdit Edit(Context.GetGeometryCollectionComponent(), GeometryCollection::EEditUpdate::RestPhysicsDynamic);
 
-				for (const int32 ClusterIndex : Context.GetSelection())
+				Context.ConvertSelectionToClusterNodes();
+				if (FGeometryCollection* GeometryCollection = Context.GetGeometryCollection().Get())
 				{
-					int32 KMeansIterations = AutoClusterSettings->ClusterSizeMethod == EClusterSizeMethod::ByGrid ? AutoClusterSettings->DriftIterations : 500;
-					// note: bPreferConvexity is incompatible with grid clustering
-					bool bPreferConvexity = AutoClusterSettings->bPreferConvexity && AutoClusterSettings->ClusterSizeMethod != EClusterSizeMethod::ByGrid;
-					FFractureEngineClustering::AutoCluster(*GeometryCollection,
-						ClusterIndex,
-						(EFractureEngineClusterSizeMethod)AutoClusterSettings->ClusterSizeMethod,
-						AutoClusterSettings->SiteCount,
-						AutoClusterSettings->SiteCountFraction,
-						AutoClusterSettings->SiteSize,
-						AutoClusterSettings->bEnforceConnectivity,
-						AutoClusterSettings->bAvoidIsolated,
-						AutoClusterSettings->bEnforceSiteParameters,
-						AutoClusterSettings->ClusterGridWidth,
-						AutoClusterSettings->ClusterGridDepth,
-						AutoClusterSettings->ClusterGridHeight,
-						AutoClusterSettings->MinimumSize,
-						KMeansIterations, 
-						bPreferConvexity, 
-						AutoClusterSettings->ConcavityTolerance
-					);
+					int32 StartTransformCount = GeometryCollection->Transform.Num();
+
+					for (const int32 ClusterIndex : Context.GetSelection())
+					{
+						int32 KMeansIterations = AutoClusterSettings->ClusterSizeMethod == EClusterSizeMethod::ByGrid ? AutoClusterSettings->DriftIterations : 500;
+						// note: bPreferConvexity is incompatible with grid clustering
+						bool bPreferConvexity = AutoClusterSettings->bPreferConvexity && AutoClusterSettings->ClusterSizeMethod != EClusterSizeMethod::ByGrid;
+						FFractureEngineClustering::AutoCluster(*GeometryCollection,
+							ClusterIndex,
+							(EFractureEngineClusterSizeMethod)AutoClusterSettings->ClusterSizeMethod,
+							AutoClusterSettings->SiteCount,
+							AutoClusterSettings->SiteCountFraction,
+							AutoClusterSettings->SiteSize,
+							AutoClusterSettings->bEnforceConnectivity,
+							AutoClusterSettings->bAvoidIsolated,
+							AutoClusterSettings->bEnforceSiteParameters,
+							AutoClusterSettings->ClusterGridWidth,
+							AutoClusterSettings->ClusterGridDepth,
+							AutoClusterSettings->ClusterGridHeight,
+							AutoClusterSettings->MinimumSize,
+							KMeansIterations,
+							bPreferConvexity,
+							AutoClusterSettings->ConcavityTolerance
+						);
+					}
+
+					Context.GenerateGuids(StartTransformCount);
 				}
 
-				Context.GenerateGuids(StartTransformCount);
+				Refresh(Context, Toolkit);
 			}
-			
-			Refresh(Context, Toolkit);
 		}
 		SetOutlinerComponents(Contexts, Toolkit);
 	}
diff --git a/Engine/Plugins/Experimental/ChaosEditor/Source/FractureEditor/Private/FractureToolClusterMagnet.cpp b/Engine/Plugins/Experimental/ChaosEditor/Source/FractureEditor/Private/FractureToolClusterMagnet.cpp
--- a/Engine/Plugins/Experimental/ChaosEditor/Source/FractureEditor/Private/FractureToolClusterMagnet.cpp
+++ b/Engine/Plugins/Experimental/ChaosEditor/Source/FractureEditor/Private/FractureToolClusterMagnet.cpp
@@ -81,18 +81,21 @@ void UFractureToolClusterMagnet::Execute(TWeakPtr<FFractureEditorModeToolkit> In
 
 		for (FFractureToolContext& Context : Contexts)
 		{
-			FGeometryCollectionEdit Edit(Context.GetGeometryCollectionComponent(), GeometryCollection::EEditUpdate::RestPhysicsDynamic);
+			if (Context.GetGeometryCollection())
+			{
+				FGeometryCollectionEdit Edit(Context.GetGeometryCollectionComponent(), GeometryCollection::EEditUpdate::RestPhysicsDynamic);
 
-			TArray<int32> Selection = Context.GetSelection();
-			int32 StartTransformCount = Context.GetGeometryCollection()->NumElements(FGeometryCollection::TransformGroup);
-			FFractureEngineClustering::ClusterMagnet(
-				*Context.GetGeometryCollection(),
-				Selection,
-				ClusterMagnetSettings->Iterations);
+				TArray<int32> Selection = Context.GetSelection();
+				int32 StartTransformCount = Context.GetGeometryCollection()->NumElements(FGeometryCollection::TransformGroup);
+				FFractureEngineClustering::ClusterMagnet(
+					*Context.GetGeometryCollection(),
+					Selection,
+					ClusterMagnetSettings->Iterations);
 
-			Context.GenerateGuids(StartTransformCount);
+				Context.GenerateGuids(StartTransformCount);
 
-			Refresh(Context, Toolkit);
+				Refresh(Context, Toolkit);
+			}
 		}
 
 		SetOutlinerComponents(Contexts, Toolkit);
diff --git a/Engine/Plugins/Experimental/ChaosEditor/Source/FractureEditor/Private/FractureToolClustering.cpp b/Engine/Plugins/Experimental/ChaosEditor/Source/FractureEditor/Private/FractureToolClustering.cpp
--- a/Engine/Plugins/Experimental/ChaosEditor/Source/FractureEditor/Private/FractureToolClustering.cpp
+++ b/Engine/Plugins/Experimental/ChaosEditor/Source/FractureEditor/Private/FractureToolClustering.cpp
@@ -47,36 +47,39 @@ void UFractureToolFlattenAll::Execute(TWeakPtr<FFractureEditorModeToolkit> InToo
 
 		for (FFractureToolContext& Context : Contexts)
 		{
-			FGeometryCollectionEdit Edit(Context.GetGeometryCollectionComponent(), GeometryCollection::EEditUpdate::RestPhysicsDynamic);
-
-			if (!Context.GetGeometryCollection()->HasAttribute("Level", FGeometryCollection::TransformGroup))
-			{
-				FGeometryCollectionClusteringUtility::UpdateHierarchyLevelOfChildren(Edit.GetRestCollection()->GetGeometryCollection().Get(), -1);
-			}
-			const TManagedArray<int32>& Levels = Context.GetGeometryCollection()->GetAttribute<int32>("Level", FGeometryCollection::TransformGroup);
-
-			Context.Sanitize();
-			Context.FilterSelectionBySimulationType(FGeometryCollection::FST_Clustered);
-
-			int32 MaxClusterLevel = -1;
-			for (int32 ClusterIndex : Context.GetSelection())
-			{
-				TArray<int32> LeafBones;
-				MaxClusterLevel = FMath::Max(MaxClusterLevel, Levels[ClusterIndex]);
-				FGeometryCollectionClusteringUtility::GetLeafBones(Context.GetGeometryCollection().Get(), ClusterIndex, true, LeafBones);
-				FGeometryCollectionClusteringUtility::ClusterBonesUnderExistingNode(Context.GetGeometryCollection().Get(), ClusterIndex, LeafBones);
-			}
-
-			// if not viewing all levels, switch the view level to show make sure we can see the cluster levels
-			if (Toolkit->GetLevelViewValue() != -1 && MaxClusterLevel != -1)
+			if (Context.GetGeometryCollection())
 			{
-				Toolkit->OnSetLevelViewValue(MaxClusterLevel + 1);
+				FGeometryCollectionEdit Edit(Context.GetGeometryCollectionComponent(), GeometryCollection::EEditUpdate::RestPhysicsDynamic);
+
+				if (!Context.GetGeometryCollection()->HasAttribute("Level", FGeometryCollection::TransformGroup))
+				{
+					FGeometryCollectionClusteringUtility::UpdateHierarchyLevelOfChildren(Edit.GetRestCollection()->GetGeometryCollection().Get(), -1);
+				}
+				const TManagedArray<int32>& Levels = Context.GetGeometryCollection()->GetAttribute<int32>("Level", FGeometryCollection::TransformGroup);
+
+				Context.Sanitize();
+				Context.FilterSelectionBySimulationType(FGeometryCollection::FST_Clustered);
+
+				int32 MaxClusterLevel = -1;
+				for (int32 ClusterIndex : Context.GetSelection())
+				{
+					TArray<int32> LeafBones;
+					MaxClusterLevel = FMath::Max(MaxClusterLevel, Levels[ClusterIndex]);
+					FGeometryCollectionClusteringUtility::GetLeafBones(Context.GetGeometryCollection().Get(), ClusterIndex, true, LeafBones);
+					FGeometryCollectionClusteringUtility::ClusterBonesUnderExistingNode(Context.GetGeometryCollection().Get(), ClusterIndex, LeafBones);
+				}
+
+				// if not viewing all levels, switch the view level to show make sure we can see the cluster levels
+				if (Toolkit->GetLevelViewValue() != -1 && MaxClusterLevel != -1)
+				{
+					Toolkit->OnSetLevelViewValue(MaxClusterLevel + 1);
+				}
+
+				// Cleanup: Remove any clusters remaining in the flattened branch.
+				FGeometryCollectionClusteringUtility::RemoveDanglingClusters(Context.GetGeometryCollection().Get());
+
+				Refresh(Context, Toolkit, true);
 			}
-
-			// Cleanup: Remove any clusters remaining in the flattened branch.
-			FGeometryCollectionClusteringUtility::RemoveDanglingClusters(Context.GetGeometryCollection().Get());
-
-			Refresh(Context, Toolkit, true);
 		}
 
 		SetOutlinerComponents(Contexts, Toolkit);
@@ -121,11 +124,14 @@ void UFractureToolCluster::Execute(TWeakPtr<FFractureEditorModeToolkit> InToolki
 
 		for (FFractureToolContext& Context : Contexts)
 		{
-			FGeometryCollectionEdit Edit(Context.GetGeometryCollectionComponent(), GeometryCollection::EEditUpdate::RestPhysicsDynamic);
-
-			if (FFractureEngineClustering::ClusterSelected(*Context.GetGeometryCollection(), Context.GetSelection()))
+			if (Context.GetGeometryCollection())
 			{
-				Refresh(Context, Toolkit);
+				FGeometryCollectionEdit Edit(Context.GetGeometryCollectionComponent(), GeometryCollection::EEditUpdate::RestPhysicsDynamic);
+
+				if (FFractureEngineClustering::ClusterSelected(*Context.GetGeometryCollection(), Context.GetSelection()))
+				{
+					Refresh(Context, Toolkit);
+				}
 			}
 		}
 
@@ -172,22 +178,25 @@ void UFractureToolUncluster::Execute(TWeakPtr<FFractureEditorModeToolkit> InTool
 
 		for (FFractureToolContext& Context : Contexts)
 		{
-			FGeometryCollectionEdit Edit(Context.GetGeometryCollectionComponent(), GeometryCollection::EEditUpdate::RestPhysicsDynamic);
-
-			const TManagedArray<TSet<int32>>& Children = Context.GetGeometryCollection()->GetAttribute<TSet<int32>>("Children", FGeometryCollection::TransformGroup);
-			if (!Context.GetGeometryCollection()->HasAttribute("Level", FGeometryCollection::TransformGroup))
+			if (Context.GetGeometryCollection())
 			{
-				FGeometryCollectionClusteringUtility::UpdateHierarchyLevelOfChildren(Edit.GetRestCollection()->GetGeometryCollection().Get(), -1);
-			}
-			const TManagedArray<int32>& Levels = Context.GetGeometryCollection()->GetAttribute<int32>("Level", FGeometryCollection::TransformGroup);
+				FGeometryCollectionEdit Edit(Context.GetGeometryCollectionComponent(), GeometryCollection::EEditUpdate::RestPhysicsDynamic);
+
+				const TManagedArray<TSet<int32>>& Children = Context.GetGeometryCollection()->GetAttribute<TSet<int32>>("Children", FGeometryCollection::TransformGroup);
+				if (!Context.GetGeometryCollection()->HasAttribute("Level", FGeometryCollection::TransformGroup))
+				{
+					FGeometryCollectionClusteringUtility::UpdateHierarchyLevelOfChildren(Edit.GetRestCollection()->GetGeometryCollection().Get(), -1);
+				}
+				const TManagedArray<int32>& Levels = Context.GetGeometryCollection()->GetAttribute<int32>("Level", FGeometryCollection::TransformGroup);
 
-			Context.ConvertSelectionToClusterNodes();
-			Context.RemoveRootNodes();
+				Context.ConvertSelectionToClusterNodes();
+				Context.RemoveRootNodes();
 
-			FGeometryCollectionClusteringUtility::CollapseHierarchyOneLevel(Context.GetGeometryCollection().Get(), Context.GetSelection());
-			
-			FGeometryCollectionClusteringUtility::RemoveDanglingClusters(Context.GetGeometryCollection().Get());
-			Refresh(Context, Toolkit, true);
+				FGeometryCollectionClusteringUtility::CollapseHierarchyOneLevel(Context.GetGeometryCollection().Get(), Context.GetSelection());
+
+				FGeometryCollectionClusteringUtility::RemoveDanglingClusters(Context.GetGeometryCollection().Get());
+				Refresh(Context, Toolkit, true);
+			}
 		}
 		
 		SetOutlinerComponents(Contexts, Toolkit);
@@ -229,14 +238,17 @@ void UFractureToolMoveUp::Execute(TWeakPtr<FFractureEditorModeToolkit> InToolkit
 
 		for (FFractureToolContext& Context: Contexts)
 		{
-			FGeometryCollectionEdit Edit(Context.GetGeometryCollectionComponent(), GeometryCollection::EEditUpdate::RestPhysicsDynamic);
-
-			Context.ConvertEmbeddedSelectionToParents(); // it's not valid to 'move up' embedded geo
-			Context.RemoveChildrenOfSelectedNodesFromSelection(); // if a node's children are also selected, just move up the parent node
-			
-			FGeometryCollectionClusteringUtility::MoveUpOneHierarchyLevel(Context.GetGeometryCollection().Get(), Context.GetSelection());
-			FGeometryCollectionClusteringUtility::RemoveDanglingClusters(Context.GetGeometryCollection().Get());
-			Refresh(Context, Toolkit, true);
+			if (Context.GetGeometryCollection())
+			{
+				FGeometryCollectionEdit Edit(Context.GetGeometryCollectionComponent(), GeometryCollection::EEditUpdate::RestPhysicsDynamic);
+
+				Context.ConvertEmbeddedSelectionToParents(); // it's not valid to 'move up' embedded geo
+				Context.RemoveChildrenOfSelectedNodesFromSelection(); // if a node's children are also selected, just move up the parent node
+
+				FGeometryCollectionClusteringUtility::MoveUpOneHierarchyLevel(Context.GetGeometryCollection().Get(), Context.GetSelection());
+				FGeometryCollectionClusteringUtility::RemoveDanglingClusters(Context.GetGeometryCollection().Get());
+				Refresh(Context, Toolkit, true);
+			}
 		}
 		
 		SetOutlinerComponents(Contexts, Toolkit);
@@ -278,11 +290,14 @@ void UFractureToolClusterMerge::Execute(TWeakPtr<FFractureEditorModeToolkit> InT
 
 		for (FFractureToolContext& Context : Contexts)
 		{
-			FGeometryCollectionEdit Edit(Context.GetGeometryCollectionComponent(), GeometryCollection::EEditUpdate::RestPhysicsDynamic);
-
-			if (FFractureEngineClustering::MergeSelectedClusters(*Context.GetGeometryCollection(), Context.GetSelection()))
+			if (Context.GetGeometryCollection())
 			{
-				Refresh(Context, Toolkit);
+				FGeometryCollectionEdit Edit(Context.GetGeometryCollectionComponent(), GeometryCollection::EEditUpdate::RestPhysicsDynamic);
+
+				if (FFractureEngineClustering::MergeSelectedClusters(*Context.GetGeometryCollection(), Context.GetSelection()))
+				{
+					Refresh(Context, Toolkit);
+				}
 			}
 		}
 
diff --git a/Engine/Plugins/Experimental/ChaosEditor/Source/FractureEditor/Private/FractureToolContext.cpp b/Engine/Plugins/Experimental/ChaosEditor/Source/FractureEditor/Private/FractureToolContext.cpp
--- a/Engine/Plugins/Experimental/ChaosEditor/Source/FractureEditor/Private/FractureToolContext.cpp
+++ b/Engine/Plugins/Experimental/ChaosEditor/Source/FractureEditor/Private/FractureToolContext.cpp
@@ -29,20 +29,23 @@ FFractureToolContext::FFractureToolContext(UGeometryCollectionComponent* InGeome
 void FFractureToolContext::Sanitize(bool bFavorParents)
 {
 	// Ensure that selected indices are valid
-	int NumTransforms = GeometryCollection->NumElements(FGeometryCollection::TransformGroup);
-	SelectedBones.RemoveAll([this, NumTransforms](int32 Index) {
-		return Index == INDEX_NONE || !ensure(Index < NumTransforms);
-		});
-	
-	// Ensure that children of a selected node are not also selected.
-	if (bFavorParents)
+	if (GeometryCollection)
 	{
-		SelectedBones.RemoveAll([this](int32 Index) {
-			return !IsValidBone(Index) || HasSelectedAncestor(Index);
+		int NumTransforms = GeometryCollection->NumElements(FGeometryCollection::TransformGroup);
+		SelectedBones.RemoveAll([this, NumTransforms](int32 Index) {
+			return Index == INDEX_NONE || !ensure(Index < NumTransforms);
 			});
-	}
 
-	SelectedBones.Sort();
+		// Ensure that children of a selected node are not also selected.
+		if (bFavorParents)
+		{
+			SelectedBones.RemoveAll([this](int32 Index) {
+				return !IsValidBone(Index) || HasSelectedAncestor(Index);
+				});
+		}
+
+		SelectedBones.Sort();
+	}
 }
 
 void FFractureToolContext::RemoveRootNodes()
diff --git a/Engine/Plugins/Experimental/ChaosEditor/Source/FractureEditor/Private/FractureToolConvert.cpp b/Engine/Plugins/Experimental/ChaosEditor/Source/FractureEditor/Private/FractureToolConvert.cpp
--- a/Engine/Plugins/Experimental/ChaosEditor/Source/FractureEditor/Private/FractureToolConvert.cpp
+++ b/Engine/Plugins/Experimental/ChaosEditor/Source/FractureEditor/Private/FractureToolConvert.cpp
@@ -105,31 +105,34 @@ TArray<FFractureToolContext> UFractureToolConvert::GetFractureToolContexts() con
 
 	for (UGeometryCollectionComponent* GeometryCollectionComponent : GeomCompSelection)
 	{
-		// Generate a context for each selected node
-		FFractureToolContext FullSelection(GeometryCollectionComponent);
-		FullSelection.ConvertSelectionToRigidNodes();
-		
-		// Update global transforms and bounds -- TODO: pull this bounds update out to a shared function?
-		const TManagedArray<FTransform3f>& Transform = FullSelection.GetGeometryCollection()->Transform;
-		const TManagedArray<int32>& TransformToGeometryIndex = FullSelection.GetGeometryCollection()->TransformToGeometryIndex;
-		const TManagedArray<FBox>& BoundingBoxes = FullSelection.GetGeometryCollection()->BoundingBox;
+		if (GeometryCollectionComponent->GetRestCollection())
+		{
+			// Generate a context for each selected node
+			FFractureToolContext FullSelection(GeometryCollectionComponent);
+			FullSelection.ConvertSelectionToRigidNodes();
 
-		TArray<FTransform> Transforms;
-		GeometryCollectionAlgo::GlobalMatrices(Transform, FullSelection.GetGeometryCollection()->Parent, Transforms);
+			// Update global transforms and bounds -- TODO: pull this bounds update out to a shared function?
+			const TManagedArray<FTransform3f>& Transform = FullSelection.GetGeometryCollection()->Transform;
+			const TManagedArray<int32>& TransformToGeometryIndex = FullSelection.GetGeometryCollection()->TransformToGeometryIndex;
+			const TManagedArray<FBox>& BoundingBoxes = FullSelection.GetGeometryCollection()->BoundingBox;
 
-		FBox Bounds(ForceInit);
-		for (int32 BoneIndex : FullSelection.GetSelection())
-		{
-			int32 GeometryIndex = TransformToGeometryIndex[BoneIndex];
-			if (GeometryIndex > INDEX_NONE)
+			TArray<FTransform> Transforms;
+			GeometryCollectionAlgo::GlobalMatrices(Transform, FullSelection.GetGeometryCollection()->Parent, Transforms);
+
+			FBox Bounds(ForceInit);
+			for (int32 BoneIndex : FullSelection.GetSelection())
 			{
-				FBox BoneBound = BoundingBoxes[GeometryIndex].TransformBy(Transforms[BoneIndex]);
-				Bounds += BoneBound;
+				int32 GeometryIndex = TransformToGeometryIndex[BoneIndex];
+				if (GeometryIndex > INDEX_NONE)
+				{
+					FBox BoneBound = BoundingBoxes[GeometryIndex].TransformBy(Transforms[BoneIndex]);
+					Bounds += BoneBound;
+				}
 			}
-		}
-		FullSelection.SetBounds(Bounds);
+			FullSelection.SetBounds(Bounds);
 
-		Contexts.Add(FullSelection);
+			Contexts.Add(FullSelection);
+		}
 	}
 
 	return Contexts;
diff --git a/Engine/Plugins/Experimental/ChaosEditor/Source/FractureEditor/Private/FractureToolEditing.cpp b/Engine/Plugins/Experimental/ChaosEditor/Source/FractureEditor/Private/FractureToolEditing.cpp
--- a/Engine/Plugins/Experimental/ChaosEditor/Source/FractureEditor/Private/FractureToolEditing.cpp
+++ b/Engine/Plugins/Experimental/ChaosEditor/Source/FractureEditor/Private/FractureToolEditing.cpp
@@ -50,54 +50,57 @@ void UFractureToolDeleteBranch::Execute(TWeakPtr<FFractureEditorModeToolkit> InT
 
 		for (FFractureToolContext& Context : Contexts)
 		{
-			FGeometryCollectionEdit Edit(Context.GetGeometryCollectionComponent(), GeometryCollection::EEditUpdate::RestPhysicsDynamic);
-			FGeometryCollection* GeometryCollection = Context.GetGeometryCollection().Get();
-			UGeometryCollection* FracturedGeometryCollection = Context.GetFracturedGeometryCollection();
+			if (Context.GetGeometryCollection())
+			{
+				FGeometryCollectionEdit Edit(Context.GetGeometryCollectionComponent(), GeometryCollection::EEditUpdate::RestPhysicsDynamic);
+				FGeometryCollection* GeometryCollection = Context.GetGeometryCollection().Get();
+				UGeometryCollection* FracturedGeometryCollection = Context.GetFracturedGeometryCollection();
 
-			const TManagedArray<int32>& ExemplarIndex = GeometryCollection->ExemplarIndex;
-			const TManagedArray<TSet<int32>>& Children = GeometryCollection->Children;
-			
-			// Removing the root node amounts to full deletion -- we don't allow this here.
-			Context.RemoveRootNodes();
-			Context.Sanitize();
+				const TManagedArray<int32>& ExemplarIndex = GeometryCollection->ExemplarIndex;
+				const TManagedArray<TSet<int32>>& Children = GeometryCollection->Children;
 
-			TArray<int32> NodesForDeletion;
+				// Removing the root node amounts to full deletion -- we don't allow this here.
+				Context.RemoveRootNodes();
+				Context.Sanitize();
 
-			for (int32 Select : Context.GetSelection())
-			{
-				FGeometryCollectionClusteringUtility::RecursiveAddAllChildren(Children, Select, NodesForDeletion);
-			}
+				TArray<int32> NodesForDeletion;
 
-			// Clean up any embedded geometry removal
-			TArray<int32> UninstancedExemplars;
-			UninstancedExemplars.Reserve(NodesForDeletion.Num());
-			for (int32 DeleteNode : NodesForDeletion)
-			{
-				if (ExemplarIndex[DeleteNode] > INDEX_NONE)
+				for (int32 Select : Context.GetSelection())
 				{
-					if ((--FracturedGeometryCollection->EmbeddedGeometryExemplar[ExemplarIndex[DeleteNode]].InstanceCount) < 1)
+					FGeometryCollectionClusteringUtility::RecursiveAddAllChildren(Children, Select, NodesForDeletion);
+				}
+
+				// Clean up any embedded geometry removal
+				TArray<int32> UninstancedExemplars;
+				UninstancedExemplars.Reserve(NodesForDeletion.Num());
+				for (int32 DeleteNode : NodesForDeletion)
+				{
+					if (ExemplarIndex[DeleteNode] > INDEX_NONE)
 					{
-						UE_LOG(LogFractureTool, Warning, TEXT("Exemplar Index %d is empty. Removing Exemplar from Geometry Collection."), ExemplarIndex[DeleteNode]);
-						UninstancedExemplars.Add(ExemplarIndex[DeleteNode]);
+						if ((--FracturedGeometryCollection->EmbeddedGeometryExemplar[ExemplarIndex[DeleteNode]].InstanceCount) < 1)
+						{
+							UE_LOG(LogFractureTool, Warning, TEXT("Exemplar Index %d is empty. Removing Exemplar from Geometry Collection."), ExemplarIndex[DeleteNode]);
+							UninstancedExemplars.Add(ExemplarIndex[DeleteNode]);
+						}
 					}
 				}
-			}
 
-			UninstancedExemplars.Sort();
-			FracturedGeometryCollection->RemoveExemplars(UninstancedExemplars);
-			GeometryCollection->ReindexExemplarIndices(UninstancedExemplars);
+				UninstancedExemplars.Sort();
+				FracturedGeometryCollection->RemoveExemplars(UninstancedExemplars);
+				GeometryCollection->ReindexExemplarIndices(UninstancedExemplars);
 
-			NodesForDeletion.Sort();
-			GeometryCollection->RemoveElements(FGeometryCollection::TransformGroup, NodesForDeletion);
+				NodesForDeletion.Sort();
+				GeometryCollection->RemoveElements(FGeometryCollection::TransformGroup, NodesForDeletion);
 
-			FGeometryCollectionClusteringUtility::RemoveDanglingClusters(GeometryCollection);
+				FGeometryCollectionClusteringUtility::RemoveDanglingClusters(GeometryCollection);
 
-			Context.GetGeometryCollectionComponent()->InitializeEmbeddedGeometry();
+				Context.GetGeometryCollectionComponent()->InitializeEmbeddedGeometry();
 
-			// Proximity is invalidated.
-			ClearProximity(Context.GetGeometryCollection().Get());
+				// Proximity is invalidated.
+				ClearProximity(Context.GetGeometryCollection().Get());
 
-			Refresh(Context, Toolkit, true);
+				Refresh(Context, Toolkit, true);
+			}
 		}
 
 		SetOutlinerComponents(Contexts, Toolkit);
@@ -139,20 +142,23 @@ void UFractureToolMergeSelected::Execute(TWeakPtr<FFractureEditorModeToolkit> In
 
 		for (FFractureToolContext& Context : Contexts)
 		{
-			FGeometryCollectionEdit Edit(Context.GetGeometryCollectionComponent(), GeometryCollection::EEditUpdate::RestPhysicsDynamic);
-			FGeometryCollection* GeometryCollection = Context.GetGeometryCollection().Get();
+			if (Context.GetGeometryCollection())
+			{
+				FGeometryCollectionEdit Edit(Context.GetGeometryCollectionComponent(), GeometryCollection::EEditUpdate::RestPhysicsDynamic);
+				FGeometryCollection* GeometryCollection = Context.GetGeometryCollection().Get();
 
-			Context.Sanitize();
+				Context.Sanitize();
 
-			const TArray<int32>& NodesForMerge = Context.GetSelection();
+				const TArray<int32>& NodesForMerge = Context.GetSelection();
 
-			constexpr bool bBooleanUnion = false;
-			MergeAllSelectedBones(*GeometryCollection, NodesForMerge, bBooleanUnion);
+				constexpr bool bBooleanUnion = false;
+				MergeAllSelectedBones(*GeometryCollection, NodesForMerge, bBooleanUnion);
 
-			// Proximity is invalidated.
-			ClearProximity(Context.GetGeometryCollection().Get());
+				// Proximity is invalidated.
+				ClearProximity(Context.GetGeometryCollection().Get());
 
-			Refresh(Context, Toolkit, true);
+				Refresh(Context, Toolkit, true);
+			}
 		}
 
 		SetOutlinerComponents(Contexts, Toolkit);
@@ -193,24 +199,27 @@ void UFractureToolSplitSelected::Execute(TWeakPtr<FFractureEditorModeToolkit> In
 
 		for (FFractureToolContext& Context : Contexts)
 		{
-			FGeometryCollectionEdit Edit(Context.GetGeometryCollectionComponent(), GeometryCollection::EEditUpdate::RestPhysicsDynamic);
-			FGeometryCollection* GeometryCollection = Context.GetGeometryCollection().Get();
+			if (Context.GetGeometryCollection())
+			{
+				FGeometryCollectionEdit Edit(Context.GetGeometryCollectionComponent(), GeometryCollection::EEditUpdate::RestPhysicsDynamic);
+				FGeometryCollection* GeometryCollection = Context.GetGeometryCollection().Get();
 
-			Context.Sanitize();
+				Context.Sanitize();
 
-			TArray<int32> NodesForSplit;
+				TArray<int32> NodesForSplit;
 
-			for (int32 Select : Context.GetSelection())
-			{
-				FGeometryCollectionClusteringUtility::GetLeafBones(GeometryCollection, Select, true, NodesForSplit);
-			}
+				for (int32 Select : Context.GetSelection())
+				{
+					FGeometryCollectionClusteringUtility::GetLeafBones(GeometryCollection, Select, true, NodesForSplit);
+				}
 
-			SplitIslands(*GeometryCollection, NodesForSplit, 0, nullptr);
+				SplitIslands(*GeometryCollection, NodesForSplit, 0, nullptr);
 
-			// Proximity is invalidated.
-			ClearProximity(Context.GetGeometryCollection().Get());
+				// Proximity is invalidated.
+				ClearProximity(Context.GetGeometryCollection().Get());
 
-			Refresh(Context, Toolkit, true);
+				Refresh(Context, Toolkit, true);
+			}
 		}
 
 		SetOutlinerComponents(Contexts, Toolkit);
@@ -252,33 +261,36 @@ void UFractureToolHide::Execute(TWeakPtr<FFractureEditorModeToolkit> InToolkit)
 
 		for (FFractureToolContext& Context : Contexts)
 		{
-			FGeometryCollectionEdit Edit(Context.GetGeometryCollectionComponent(), GeometryCollection::EEditUpdate::RestPhysicsDynamic);
-			FGeometryCollection* GeometryCollection = Edit.GetRestCollection()->GetGeometryCollection().Get();
-			UGeometryCollection* FracturedGeometryCollection = Context.GetFracturedGeometryCollection();
+			if (Context.GetGeometryCollection())
+			{
+				FGeometryCollectionEdit Edit(Context.GetGeometryCollectionComponent(), GeometryCollection::EEditUpdate::RestPhysicsDynamic);
+				FGeometryCollection* GeometryCollection = Edit.GetRestCollection()->GetGeometryCollection().Get();
+				UGeometryCollection* FracturedGeometryCollection = Context.GetFracturedGeometryCollection();
 
-			Context.ConvertSelectionToRigidNodes();
+				Context.ConvertSelectionToRigidNodes();
 
-			const TManagedArray<int32>&	TransformToGeometryIndex = GeometryCollection->TransformToGeometryIndex;
-			const TManagedArray<int32>&	FaceStart = GeometryCollection->FaceStart;
-			const TManagedArray<int32>&	FaceCount = GeometryCollection->FaceCount;
-			TManagedArray<bool>& Visible = GeometryCollection->Visible;
+				const TManagedArray<int32>& TransformToGeometryIndex = GeometryCollection->TransformToGeometryIndex;
+				const TManagedArray<int32>& FaceStart = GeometryCollection->FaceStart;
+				const TManagedArray<int32>& FaceCount = GeometryCollection->FaceCount;
+				TManagedArray<bool>& Visible = GeometryCollection->Visible;
 
 
-			const TArray<int32>& Selection = Context.GetSelection();
-			for (int32 Idx : Selection)
-			{
-				// Iterate the faces in the geometry of this rigid node and set invisible.
-				if (TransformToGeometryIndex[Idx] > INDEX_NONE)
+				const TArray<int32>& Selection = Context.GetSelection();
+				for (int32 Idx : Selection)
 				{
-					int32 CurrFace = FaceStart[TransformToGeometryIndex[Idx]];
-					for (int32 FaceOffset = 0; FaceOffset < FaceCount[TransformToGeometryIndex[Idx]]; ++FaceOffset)
+					// Iterate the faces in the geometry of this rigid node and set invisible.
+					if (TransformToGeometryIndex[Idx] > INDEX_NONE)
 					{
-						Visible[CurrFace + FaceOffset] = false;
+						int32 CurrFace = FaceStart[TransformToGeometryIndex[Idx]];
+						for (int32 FaceOffset = 0; FaceOffset < FaceCount[TransformToGeometryIndex[Idx]]; ++FaceOffset)
+						{
+							Visible[CurrFace + FaceOffset] = false;
+						}
 					}
 				}
-			}
 
-			Refresh(Context, Toolkit, true);
+				Refresh(Context, Toolkit, true);
+			}
 		}
 
 		SetOutlinerComponents(Contexts, Toolkit);
@@ -320,33 +332,36 @@ void UFractureToolUnhide::Execute(TWeakPtr<FFractureEditorModeToolkit> InToolkit
 
 		for (FFractureToolContext& Context : Contexts)
 		{
-			FGeometryCollectionEdit Edit(Context.GetGeometryCollectionComponent(), GeometryCollection::EEditUpdate::RestPhysicsDynamic);
-			FGeometryCollection* GeometryCollection = Edit.GetRestCollection()->GetGeometryCollection().Get();
-			UGeometryCollection* FracturedGeometryCollection = Context.GetFracturedGeometryCollection();
+			if (Context.GetGeometryCollection())
+			{
+				FGeometryCollectionEdit Edit(Context.GetGeometryCollectionComponent(), GeometryCollection::EEditUpdate::RestPhysicsDynamic);
+				FGeometryCollection* GeometryCollection = Edit.GetRestCollection()->GetGeometryCollection().Get();
+				UGeometryCollection* FracturedGeometryCollection = Context.GetFracturedGeometryCollection();
 
-			Context.ConvertSelectionToRigidNodes();
+				Context.ConvertSelectionToRigidNodes();
 
-			const TManagedArray<int32>& TransformToGeometryIndex = GeometryCollection->TransformToGeometryIndex;
-			const TManagedArray<int32>& FaceStart = GeometryCollection->FaceStart;
-			const TManagedArray<int32>& FaceCount = GeometryCollection->FaceCount;
-			TManagedArray<bool>& Visible = GeometryCollection->Visible;
+				const TManagedArray<int32>& TransformToGeometryIndex = GeometryCollection->TransformToGeometryIndex;
+				const TManagedArray<int32>& FaceStart = GeometryCollection->FaceStart;
+				const TManagedArray<int32>& FaceCount = GeometryCollection->FaceCount;
+				TManagedArray<bool>& Visible = GeometryCollection->Visible;
 
 
-			const TArray<int32>& Selection = Context.GetSelection();
-			for (int32 Idx : Selection)
-			{
-				// Iterate the faces in the geometry of this rigid node and set invisible.
-				if (TransformToGeometryIndex[Idx] > INDEX_NONE)
+				const TArray<int32>& Selection = Context.GetSelection();
+				for (int32 Idx : Selection)
 				{
-					int32 CurrFace = FaceStart[TransformToGeometryIndex[Idx]];
-					for (int32 FaceOffset = 0; FaceOffset < FaceCount[TransformToGeometryIndex[Idx]]; ++FaceOffset)
+					// Iterate the faces in the geometry of this rigid node and set invisible.
+					if (TransformToGeometryIndex[Idx] > INDEX_NONE)
 					{
-						Visible[CurrFace + FaceOffset] = true;
+						int32 CurrFace = FaceStart[TransformToGeometryIndex[Idx]];
+						for (int32 FaceOffset = 0; FaceOffset < FaceCount[TransformToGeometryIndex[Idx]]; ++FaceOffset)
+						{
+							Visible[CurrFace + FaceOffset] = true;
+						}
 					}
 				}
-			}
 
-			Refresh(Context, Toolkit, true);
+				Refresh(Context, Toolkit, true);
+			}
 		}
 
 		SetOutlinerComponents(Contexts, Toolkit);
diff --git a/Engine/Plugins/Experimental/ChaosEditor/Source/FractureEditor/Private/FractureToolEmbed.cpp b/Engine/Plugins/Experimental/ChaosEditor/Source/FractureEditor/Private/FractureToolEmbed.cpp
--- a/Engine/Plugins/Experimental/ChaosEditor/Source/FractureEditor/Private/FractureToolEmbed.cpp
+++ b/Engine/Plugins/Experimental/ChaosEditor/Source/FractureEditor/Private/FractureToolEmbed.cpp
@@ -63,49 +63,52 @@ void UFractureToolAddEmbeddedGeometry::Execute(TWeakPtr<FFractureEditorModeToolk
 
 		for (FFractureToolContext& Context : Contexts)
 		{
-			FGeometryCollectionEdit GeometryCollectionEdit = Context.GetGeometryCollectionComponent()->EditRestCollection(GeometryCollection::EEditUpdate::RestPhysicsDynamic, true /*bShapeIsUnchanged*/);
-			UGeometryCollection* FracturedGeometryCollection = GeometryCollectionEdit.GetRestCollection();
+			if (Context.GetGeometryCollection())
+			{
+				FGeometryCollectionEdit GeometryCollectionEdit = Context.GetGeometryCollectionComponent()->EditRestCollection(GeometryCollection::EEditUpdate::RestPhysicsDynamic, true /*bShapeIsUnchanged*/);
+				UGeometryCollection* FracturedGeometryCollection = GeometryCollectionEdit.GetRestCollection();
 
-			FFractureToolContext::FGeometryCollectionPtr GeometryCollection = Context.GetGeometryCollection();
-			const TManagedArray<FTransform3f>& Transform = GeometryCollection->Transform;
-			const TManagedArray<int32>& Parent = GeometryCollection->Parent;
+				FFractureToolContext::FGeometryCollectionPtr GeometryCollection = Context.GetGeometryCollection();
+				const TManagedArray<FTransform3f>& Transform = GeometryCollection->Transform;
+				const TManagedArray<int32>& Parent = GeometryCollection->Parent;
 
-			int32 StartTransformCount = Transform.Num();
-			
-			const FTransform TargetActorTransform(Context.GetGeometryCollectionComponent()->GetOwner()->GetTransform());
+				int32 StartTransformCount = Transform.Num();
 
-			Context.ConvertSelectionToRigidNodes();
-			const TArray<int32>& SelectedBones = Context.GetSelection();
-			for (const int32 SelectedBone : SelectedBones)
-			{
-				FTransform BoneGlobalTransform = GeometryCollectionAlgo::GlobalMatrix(Transform, Parent, SelectedBone);
+				const FTransform TargetActorTransform(Context.GetGeometryCollectionComponent()->GetOwner()->GetTransform());
 
-				for (UStaticMeshComponent* SelectedStaticMeshComponent : SelectedStaticMeshComponents)
+				Context.ConvertSelectionToRigidNodes();
+				const TArray<int32>& SelectedBones = Context.GetSelection();
+				for (const int32 SelectedBone : SelectedBones)
 				{
-					const AActor* Actor = SelectedStaticMeshComponent->GetOwner();
-					check(Actor)
+					FTransform BoneGlobalTransform = GeometryCollectionAlgo::GlobalMatrix(Transform, Parent, SelectedBone);
 
-					const FTransform SMActorTransform(Actor->GetTransform());
+					for (UStaticMeshComponent* SelectedStaticMeshComponent : SelectedStaticMeshComponents)
+					{
+						const AActor* Actor = SelectedStaticMeshComponent->GetOwner();
+						check(Actor)
 
-					UStaticMesh* ComponentStaticMesh = SelectedStaticMeshComponent->GetStaticMesh();
-					
-					FTransform ComponentTransform = SMActorTransform.GetRelativeTransform(TargetActorTransform);
-					FTransform BoneTransform = ComponentTransform.GetRelativeTransform(BoneGlobalTransform);
+							const FTransform SMActorTransform(Actor->GetTransform());
 
-					int32 ExemplarIndex = FracturedGeometryCollection->AttachEmbeddedGeometryExemplar(ComponentStaticMesh);
-					if (GeometryCollection->AppendEmbeddedInstance(ExemplarIndex, SelectedBone, BoneTransform))
-					{
-						FracturedGeometryCollection->EmbeddedGeometryExemplar[ExemplarIndex].InstanceCount++;
+						UStaticMesh* ComponentStaticMesh = SelectedStaticMeshComponent->GetStaticMesh();
+
+						FTransform ComponentTransform = SMActorTransform.GetRelativeTransform(TargetActorTransform);
+						FTransform BoneTransform = ComponentTransform.GetRelativeTransform(BoneGlobalTransform);
+
+						int32 ExemplarIndex = FracturedGeometryCollection->AttachEmbeddedGeometryExemplar(ComponentStaticMesh);
+						if (GeometryCollection->AppendEmbeddedInstance(ExemplarIndex, SelectedBone, BoneTransform))
+						{
+							FracturedGeometryCollection->EmbeddedGeometryExemplar[ExemplarIndex].InstanceCount++;
+						}
 					}
 				}
-			}
 
-			Context.GenerateGuids(StartTransformCount);
+				Context.GenerateGuids(StartTransformCount);
 
-			Context.GetGeometryCollectionComponent()->InitializeEmbeddedGeometry();
-			Refresh(Context, Toolkit);
+				Context.GetGeometryCollectionComponent()->InitializeEmbeddedGeometry();
+				Refresh(Context, Toolkit);
 
-			FracturedGeometryCollection->MarkPackageDirty();
+				FracturedGeometryCollection->MarkPackageDirty();
+			}
 		}
 
 		SetOutlinerComponents(Contexts, Toolkit);
@@ -191,40 +194,43 @@ void UFractureToolAutoEmbedGeometry::Execute(TWeakPtr<FFractureEditorModeToolkit
 
 			for (FFractureToolContext& Context : Contexts)
 			{
-				FGeometryCollectionEdit Edit(Context.GetGeometryCollectionComponent(), GeometryCollection::EEditUpdate::RestPhysicsDynamic, true /*bShapeIsUnchanged*/);
+				if (Context.GetGeometryCollection())
+				{
+					FGeometryCollectionEdit Edit(Context.GetGeometryCollectionComponent(), GeometryCollection::EEditUpdate::RestPhysicsDynamic, true /*bShapeIsUnchanged*/);
 
-				FGeometryCollection* GeometryCollection = Context.GetGeometryCollection().Get();
-				FGeometryCollectionConvexUtility::FGeometryCollectionConvexData ConvexData = FGeometryCollectionConvexUtility::GetValidConvexHullData(GeometryCollection);
-				const TManagedArray<FTransform3f>& Transform = GeometryCollection->Transform;
-				const TManagedArray<int32>& Parent = GeometryCollection->Parent;
-				const TManagedArray<int32>& SimulationType = GeometryCollection->SimulationType;
-				TArray<FTransform> BoneGlobalTransforms;
-				GeometryCollectionAlgo::GlobalMatrices(Transform, Parent, BoneGlobalTransforms);
+					FGeometryCollection* GeometryCollection = Context.GetGeometryCollection().Get();
+					FGeometryCollectionConvexUtility::FGeometryCollectionConvexData ConvexData = FGeometryCollectionConvexUtility::GetValidConvexHullData(GeometryCollection);
+					const TManagedArray<FTransform3f>& Transform = GeometryCollection->Transform;
+					const TManagedArray<int32>& Parent = GeometryCollection->Parent;
+					const TManagedArray<int32>& SimulationType = GeometryCollection->SimulationType;
+					TArray<FTransform> BoneGlobalTransforms;
+					GeometryCollectionAlgo::GlobalMatrices(Transform, Parent, BoneGlobalTransforms);
 
-				FTransform WorldToComponent = Context.GetGeometryCollectionComponent()->GetComponentToWorld().Inverse();
-				FVector ComponentSpaceLocation = WorldToComponent.TransformPosition(SMLocation);
+					FTransform WorldToComponent = Context.GetGeometryCollectionComponent()->GetComponentToWorld().Inverse();
+					FVector ComponentSpaceLocation = WorldToComponent.TransformPosition(SMLocation);
 
-				int32 NumConvex = ConvexData.ConvexHull.Num();
-				for (int32 TransformIndex = 0; TransformIndex < Transform.Num(); TransformIndex++)
-				{
-					if (SimulationType[TransformIndex] != FGeometryCollection::ESimulationTypes::FST_Rigid)
+					int32 NumConvex = ConvexData.ConvexHull.Num();
+					for (int32 TransformIndex = 0; TransformIndex < Transform.Num(); TransformIndex++)
 					{
-						continue;
-					}
+						if (SimulationType[TransformIndex] != FGeometryCollection::ESimulationTypes::FST_Rigid)
+						{
+							continue;
+						}
 
-					FVector BoneSpaceLocation = BoneGlobalTransforms[TransformIndex].InverseTransformPosition(ComponentSpaceLocation);
-					for (int32 ConvexIndex : ConvexData.TransformToConvexIndices[TransformIndex])
-					{
-						Chaos::FVec3 Normal;
-						Chaos::FReal Phi = ConvexData.ConvexHull[ConvexIndex]->PhiWithNormal(BoneSpaceLocation, Normal);
-						UE_LOG(LogTemp, Warning, TEXT(" Cone %d Phi %f"), TransformIndex, Phi);
-						if (Phi < ClosestPhi)
+						FVector BoneSpaceLocation = BoneGlobalTransforms[TransformIndex].InverseTransformPosition(ComponentSpaceLocation);
+						for (int32 ConvexIndex : ConvexData.TransformToConvexIndices[TransformIndex])
 						{
-							ClosestPhi = Phi;
-							ClosestGeometryCollection = Context.GetGeometryCollection();
-							ClosestConvex = ConvexIndex;
-							ClosestComponent = Context.GetGeometryCollectionComponent();
-							ClosestContext = &Context;
+							Chaos::FVec3 Normal;
+							Chaos::FReal Phi = ConvexData.ConvexHull[ConvexIndex]->PhiWithNormal(BoneSpaceLocation, Normal);
+							UE_LOG(LogTemp, Warning, TEXT(" Cone %d Phi %f"), TransformIndex, Phi);
+							if (Phi < ClosestPhi)
+							{
+								ClosestPhi = Phi;
+								ClosestGeometryCollection = Context.GetGeometryCollection();
+								ClosestConvex = ConvexIndex;
+								ClosestComponent = Context.GetGeometryCollectionComponent();
+								ClosestContext = &Context;
+							}
 						}
 					}
 				}
@@ -372,78 +378,81 @@ void UFractureToolFlushEmbeddedGeometry::Execute(TWeakPtr<FFractureEditorModeToo
 
 		for (FFractureToolContext& Context : Contexts)
 		{
-			if (!Context.GetGeometryCollection()->HasAttribute("ExemplarIndex", FGeometryCollection::TransformGroup))
+			if (Context.GetGeometryCollection())
 			{
-				continue;
-			}
-			
-			const TManagedArray<int32>& ExemplarIndex = Context.GetGeometryCollection()->GetAttribute<int32>("ExemplarIndex", FGeometryCollection::TransformGroup);
-			const TManagedArray<FTransform3f>& Transforms = Context.GetGeometryCollection()->GetAttribute<FTransform3f>("Transform", FGeometryCollection::TransformGroup);
-			const TManagedArray<int32>& ParentIndices = Context.GetGeometryCollection()->GetAttribute<int32>("Parent", FGeometryCollection::TransformGroup);
+				if (!Context.GetGeometryCollection()->HasAttribute("ExemplarIndex", FGeometryCollection::TransformGroup))
+				{
+					continue;
+				}
 
-			TArray<FMatrix> GlobalMatrices;
-			GeometryCollectionAlgo::GlobalMatrices(Transforms, ParentIndices, GlobalMatrices);
-			
-			TArray<int32> RemovalList;
-			RemovalList.Reserve(ExemplarIndex.Num());
+				const TManagedArray<int32>& ExemplarIndex = Context.GetGeometryCollection()->GetAttribute<int32>("ExemplarIndex", FGeometryCollection::TransformGroup);
+				const TManagedArray<FTransform3f>& Transforms = Context.GetGeometryCollection()->GetAttribute<FTransform3f>("Transform", FGeometryCollection::TransformGroup);
+				const TManagedArray<int32>& ParentIndices = Context.GetGeometryCollection()->GetAttribute<int32>("Parent", FGeometryCollection::TransformGroup);
 
-			TMap<int32,TArray<FTransform>> ExemplarTransforms;
+				TArray<FMatrix> GlobalMatrices;
+				GeometryCollectionAlgo::GlobalMatrices(Transforms, ParentIndices, GlobalMatrices);
 
-			// Filter for instanced embedded geometry
-			for (int32 TransformIndex = 0; TransformIndex < ExemplarIndex.Num(); ++TransformIndex)
-			{
-				if (ExemplarIndex[TransformIndex] > INDEX_NONE)
+				TArray<int32> RemovalList;
+				RemovalList.Reserve(ExemplarIndex.Num());
+
+				TMap<int32, TArray<FTransform>> ExemplarTransforms;
+
+				// Filter for instanced embedded geometry
+				for (int32 TransformIndex = 0; TransformIndex < ExemplarIndex.Num(); ++TransformIndex)
 				{
-					RemovalList.Add(TransformIndex);
-					TArray<FTransform>& InstanceTransforms = ExemplarTransforms.FindOrAdd(ExemplarIndex[TransformIndex]);
-					InstanceTransforms.Add(FTransform(GlobalMatrices[TransformIndex]));
+					if (ExemplarIndex[TransformIndex] > INDEX_NONE)
+					{
+						RemovalList.Add(TransformIndex);
+						TArray<FTransform>& InstanceTransforms = ExemplarTransforms.FindOrAdd(ExemplarIndex[TransformIndex]);
+						InstanceTransforms.Add(FTransform(GlobalMatrices[TransformIndex]));
+					}
 				}
-			}
-
-			if (RemovalList.Num())
-			{
-				// Remove all instanced embedded geometry.
-				FGeometryCollectionEdit GeometryCollectionEdit = Context.GetGeometryCollectionComponent()->EditRestCollection(GeometryCollection::EEditUpdate::RestPhysicsDynamic, true /*bShapeIsUnchanged*/);
-				UGeometryCollection* FracturedCollection = GeometryCollectionEdit.GetRestCollection();
-				FracturedCollection->GetGeometryCollection()->RemoveElements(FGeometryCollection::TransformGroup, RemovalList);
 
-				if (FlushEmbeddedGeometrySettings->bExtractAsStaticMeshActors)
+				if (RemovalList.Num())
 				{
-					FActorSpawnParameters SpawnParameters;
-					
-					FTransform ComponentSpace = Context.GetGeometryCollectionComponent()->GetComponentTransform();
-					
-					if (UWorld* TargetWorld = Context.GetGeometryCollectionComponent()->GetOwner()->GetWorld())
+					// Remove all instanced embedded geometry.
+					FGeometryCollectionEdit GeometryCollectionEdit = Context.GetGeometryCollectionComponent()->EditRestCollection(GeometryCollection::EEditUpdate::RestPhysicsDynamic, true /*bShapeIsUnchanged*/);
+					UGeometryCollection* FracturedCollection = GeometryCollectionEdit.GetRestCollection();
+					FracturedCollection->GetGeometryCollection()->RemoveElements(FGeometryCollection::TransformGroup, RemovalList);
+
+					if (FlushEmbeddedGeometrySettings->bExtractAsStaticMeshActors)
 					{
-						// Place Static Mesh Actors
-						for (auto It = ExemplarTransforms.CreateConstIterator(); It; ++It)
+						FActorSpawnParameters SpawnParameters;
+
+						FTransform ComponentSpace = Context.GetGeometryCollectionComponent()->GetComponentTransform();
+
+						if (UWorld* TargetWorld = Context.GetGeometryCollectionComponent()->GetOwner()->GetWorld())
 						{
-							FSoftObjectPath ExemplarPath = FracturedCollection->EmbeddedGeometryExemplar[It->Key].StaticMeshExemplar;
-							UObject* LoadedExemplarAsset = ExemplarPath.TryLoad();
-							if (UActorFactory* ActorFactory = FActorFactoryAssetProxy::GetFactoryForAssetObject(LoadedExemplarAsset))
+							// Place Static Mesh Actors
+							for (auto It = ExemplarTransforms.CreateConstIterator(); It; ++It)
 							{
-
-								for (const FTransform& Transform : It->Value)
+								FSoftObjectPath ExemplarPath = FracturedCollection->EmbeddedGeometryExemplar[It->Key].StaticMeshExemplar;
+								UObject* LoadedExemplarAsset = ExemplarPath.TryLoad();
+								if (UActorFactory* ActorFactory = FActorFactoryAssetProxy::GetFactoryForAssetObject(LoadedExemplarAsset))
 								{
-									ActorFactory->CreateActor(LoadedExemplarAsset, TargetWorld->GetCurrentLevel(), Transform * ComponentSpace, SpawnParameters);
+
+									for (const FTransform& Transform : It->Value)
+									{
+										ActorFactory->CreateActor(LoadedExemplarAsset, TargetWorld->GetCurrentLevel(), Transform * ComponentSpace, SpawnParameters);
+									}
 								}
-							}
 
+							}
 						}
+
 					}
-					
-				}
-								
-				// Remove all Exemplars from asset.
-				FracturedCollection->EmbeddedGeometryExemplar.Empty();
-				
-				Context.GetGeometryCollectionComponent()->InitializeEmbeddedGeometry();
 
-				FScopedColorEdit ScopedEdit(Context.GetGeometryCollectionComponent());
-				ScopedEdit.ResetBoneSelection();
-				Refresh(Context, Toolkit);
-				FracturedCollection->MarkPackageDirty();
-			}	
+					// Remove all Exemplars from asset.
+					FracturedCollection->EmbeddedGeometryExemplar.Empty();
+
+					Context.GetGeometryCollectionComponent()->InitializeEmbeddedGeometry();
+
+					FScopedColorEdit ScopedEdit(Context.GetGeometryCollectionComponent());
+					ScopedEdit.ResetBoneSelection();
+					Refresh(Context, Toolkit);
+					FracturedCollection->MarkPackageDirty();
+				}
+			}
 		}
 
 		SetOutlinerComponents(Contexts, Toolkit);
diff --git a/Engine/Plugins/Experimental/ChaosEditor/Source/FractureEditor/Private/FractureToolFixTinyGeo.cpp b/Engine/Plugins/Experimental/ChaosEditor/Source/FractureEditor/Private/FractureToolFixTinyGeo.cpp
--- a/Engine/Plugins/Experimental/ChaosEditor/Source/FractureEditor/Private/FractureToolFixTinyGeo.cpp
+++ b/Engine/Plugins/Experimental/ChaosEditor/Source/FractureEditor/Private/FractureToolFixTinyGeo.cpp
@@ -182,48 +182,51 @@ void UFractureToolFixTinyGeo::FractureContextChanged()
 	TArray<int32> Leaves;
 	for (UGeometryCollectionComponent* GeometryCollectionComponent : GeomCompSelection)
 	{
-		// Get a non-const rest collection to support update of the Volume attribute
-		FGeometryCollectionEdit EditCollection = GeometryCollectionComponent->EditRestCollection(GeometryCollection::EEditUpdate::None, true);
-		FGeometryCollection& Collection = *EditCollection.GetRestCollection()->GetGeometryCollection();
-		GeometryCollection::Facades::FCollectionTransformSelectionFacade SelectionFacade(Collection);
-
-		TArray<int32> SelectedBones = GeometryCollectionComponent->GetSelectedBones();
-		
-		int CollectionIdx = VisualizedCollections.Emplace(GeometryCollectionComponent);
-
-		TArray<int32> SmallIndices;
-		TArray<double> Volumes;
-		double MinVolume;
-		CollectTargetBones(Collection, SelectedBones, SmallIndices, Volumes, MinVolume);
-
-		FTransform OuterTransform = GeometryCollectionComponent->GetOwner()->GetActorTransform();
-		for (int32 TransformIdx : SmallIndices) // small transforms
+		if (GeometryCollectionComponent->GetRestCollection())
 		{
-			FTransform InnerTransform = GeometryCollectionAlgo::GlobalMatrix(Collection.Transform, Collection.Parent, TransformIdx);
+			// Get a non-const rest collection to support update of the Volume attribute
+			FGeometryCollectionEdit EditCollection = GeometryCollectionComponent->EditRestCollection(GeometryCollection::EEditUpdate::None, true);
+			FGeometryCollection& Collection = *EditCollection.GetRestCollection()->GetGeometryCollection();
+			GeometryCollection::Facades::FCollectionTransformSelectionFacade SelectionFacade(Collection);
 
+			TArray<int32> SelectedBones = GeometryCollectionComponent->GetSelectedBones();
 
-			FTransform CombinedTransform = InnerTransform * OuterTransform;
+			int CollectionIdx = VisualizedCollections.Emplace(GeometryCollectionComponent);
 
-			FBox Bounds(EForceInit::ForceInit);
-			Leaves.Reset();
-			SelectionFacade.ConvertSelectionToRigidNodes(TransformIdx, Leaves);
-			for (int32 LeafIdx : Leaves)
+			TArray<int32> SmallIndices;
+			TArray<double> Volumes;
+			double MinVolume;
+			CollectTargetBones(Collection, SelectedBones, SmallIndices, Volumes, MinVolume);
+
+			FTransform OuterTransform = GeometryCollectionComponent->GetOwner()->GetActorTransform();
+			for (int32 TransformIdx : SmallIndices) // small transforms
 			{
-				int32 GeometryIdx = Collection.TransformToGeometryIndex[LeafIdx];
-				if (GeometryIdx == INDEX_NONE)
-				{
-					ensure(false); // already filtered for small geo, shouldn't have any missing-geo transforms in the list
-					continue;
-				}
-				int32 VStart = Collection.VertexStart[GeometryIdx];
-				int32 VEnd = VStart + Collection.VertexCount[GeometryIdx];
-				for (int32 VIdx = VStart; VIdx < VEnd; VIdx++)
+				FTransform InnerTransform = GeometryCollectionAlgo::GlobalMatrix(Collection.Transform, Collection.Parent, TransformIdx);
+
+
+				FTransform CombinedTransform = InnerTransform * OuterTransform;
+
+				FBox Bounds(EForceInit::ForceInit);
+				Leaves.Reset();
+				SelectionFacade.ConvertSelectionToRigidNodes(TransformIdx, Leaves);
+				for (int32 LeafIdx : Leaves)
 				{
-					Bounds += CombinedTransform.TransformPosition((FVector)Collection.Vertex[VIdx]);
+					int32 GeometryIdx = Collection.TransformToGeometryIndex[LeafIdx];
+					if (GeometryIdx == INDEX_NONE)
+					{
+						ensure(false); // already filtered for small geo, shouldn't have any missing-geo transforms in the list
+						continue;
+					}
+					int32 VStart = Collection.VertexStart[GeometryIdx];
+					int32 VEnd = VStart + Collection.VertexCount[GeometryIdx];
+					for (int32 VIdx = VStart; VIdx < VEnd; VIdx++)
+					{
+						Bounds += CombinedTransform.TransformPosition((FVector)Collection.Vertex[VIdx]);
+					}
 				}
+				ToRemoveMappings.AddMapping(CollectionIdx, TransformIdx, ToRemoveBounds.Num());
+				ToRemoveBounds.Add(Bounds);
 			}
-			ToRemoveMappings.AddMapping(CollectionIdx, TransformIdx, ToRemoveBounds.Num());
-			ToRemoveBounds.Add(Bounds);
 		}
 	}
 }
@@ -273,53 +276,56 @@ void UFractureToolFixTinyGeo::Execute(TWeakPtr<FFractureEditorModeToolkit> InToo
 	GetSelectedGeometryCollectionComponents(GeomCompSelection);
 	for (UGeometryCollectionComponent* GeometryCollectionComponent : GeomCompSelection)
 	{
-		FGeometryCollectionEdit CollectionEdit = GeometryCollectionComponent->EditRestCollection(GeometryCollection::EEditUpdate::RestPhysicsDynamic, false);
-		FGeometryCollection& Collection = *CollectionEdit.GetRestCollection()->GetGeometryCollection();
-		TArray<int32> SelectedBones = GeometryCollectionComponent->GetSelectedBones();
-
-		TArray<int32> SmallIndices;
-		TArray<double> Volumes;
-		double MinVolume;
-		if (!CollectTargetBones(Collection, SelectedBones, SmallIndices, Volumes, MinVolume))
+		if (GeometryCollectionComponent->GetRestCollection())
 		{
-			continue;
-		}
+			FGeometryCollectionEdit CollectionEdit = GeometryCollectionComponent->EditRestCollection(GeometryCollection::EEditUpdate::RestPhysicsDynamic, false);
+			FGeometryCollection& Collection = *CollectionEdit.GetRestCollection()->GetGeometryCollection();
+			TArray<int32> SelectedBones = GeometryCollectionComponent->GetSelectedBones();
+
+			TArray<int32> SmallIndices;
+			TArray<double> Volumes;
+			double MinVolume;
+			if (!CollectTargetBones(Collection, SelectedBones, SmallIndices, Volumes, MinVolume))
+			{
+				continue;
+			}
 
-		// convert ENeighborSelectionMethod to the PlanarCut module version (the only difference is it's not a UENUM)
-		UE::PlanarCut::ENeighborSelectionMethod SelectionMethod = UE::PlanarCut::ENeighborSelectionMethod::LargestNeighbor;
-		if (TinyGeoSettings->NeighborSelection == ENeighborSelectionMethod::NearestCenter)
-		{
-			SelectionMethod = UE::PlanarCut::ENeighborSelectionMethod::NearestCenter;
-		}
+			// convert ENeighborSelectionMethod to the PlanarCut module version (the only difference is it's not a UENUM)
+			UE::PlanarCut::ENeighborSelectionMethod SelectionMethod = UE::PlanarCut::ENeighborSelectionMethod::LargestNeighbor;
+			if (TinyGeoSettings->NeighborSelection == ENeighborSelectionMethod::NearestCenter)
+			{
+				SelectionMethod = UE::PlanarCut::ENeighborSelectionMethod::NearestCenter;
+			}
 
-		if (TinyGeoSettings->MergeType == EMergeType::MergeGeometry)
-		{
-			MergeBones(
-				Collection,
-				TArrayView<const int32>(), // empty view == consider all bones
-				Volumes,
-				MinVolume,
-				SmallIndices,
-				false, /*bUnionJoinedPieces*/ // Note: Union-ing the pieces is nicer in theory, but can leave cracks and non-manifold garbage
-				SelectionMethod
-			);
-		}
-		else
-		{
-			MergeClusters(
-				Collection,
-				Volumes,
-				MinVolume,
-				SmallIndices,
-				SelectionMethod,
-				TinyGeoSettings->bOnlyToConnected,
-				TinyGeoSettings->bOnlySameParent
-			);
-		}
+			if (TinyGeoSettings->MergeType == EMergeType::MergeGeometry)
+			{
+				MergeBones(
+					Collection,
+					TArrayView<const int32>(), // empty view == consider all bones
+					Volumes,
+					MinVolume,
+					SmallIndices,
+					false, /*bUnionJoinedPieces*/ // Note: Union-ing the pieces is nicer in theory, but can leave cracks and non-manifold garbage
+					SelectionMethod
+				);
+			}
+			else
+			{
+				MergeClusters(
+					Collection,
+					Volumes,
+					MinVolume,
+					SmallIndices,
+					SelectionMethod,
+					TinyGeoSettings->bOnlyToConnected,
+					TinyGeoSettings->bOnlySameParent
+				);
+			}
 
-		// Refresh component with cleared selection
-		TArray<int32> EmptySelection;
-		Refresh(GeometryCollectionComponent, Toolkit, EmptySelection, true /*bClearSelection*/, true /*bMustUpdateBoneColors*/);
+			// Refresh component with cleared selection
+			TArray<int32> EmptySelection;
+			Refresh(GeometryCollectionComponent, Toolkit, EmptySelection, true /*bClearSelection*/, true /*bMustUpdateBoneColors*/);
+		}
 	}
 
 	// Update overall outliner (including stats, etc)
diff --git a/Engine/Plugins/Experimental/ChaosEditor/Source/FractureEditor/Private/FractureToolProximity.cpp b/Engine/Plugins/Experimental/ChaosEditor/Source/FractureEditor/Private/FractureToolProximity.cpp
--- a/Engine/Plugins/Experimental/ChaosEditor/Source/FractureEditor/Private/FractureToolProximity.cpp
+++ b/Engine/Plugins/Experimental/ChaosEditor/Source/FractureEditor/Private/FractureToolProximity.cpp
@@ -104,63 +104,66 @@ void UFractureToolProximity::UpdateVisualizations()
 
 	for (const FFractureToolContext& FractureContext : FractureContexts)
 	{
-		FGeometryCollection& Collection = *FractureContext.GetGeometryCollection();
-
-		if (!Collection.HasAttribute("Proximity", FGeometryCollection::GeometryGroup))
-		{
-			continue;
-		}
-		const TManagedArray<TSet<int32>>& Proximity = Collection.GetAttribute<TSet<int32>>("Proximity", FGeometryCollection::GeometryGroup);
-
-		// Save the geometry collection component for rendering
-		int32 CollectionIdx = VisualizedCollections.Add(FractureContext.GetGeometryCollectionComponent());
-		FCollectionVisInfo& Vis = ProximityVisualizations.Emplace_GetRef();
-		Vis.CollectionIndex = CollectionIdx;
-
-		// Get proximity graph node positions as bounding box centers for each piece of geometry
-		Vis.GeoCenters.SetNum(Collection.NumElements(FGeometryCollection::GeometryGroup));
-		TArray<FTransform> GlobalTransformArray;
-		GeometryCollectionAlgo::GlobalMatrices(Collection.Transform, Collection.Parent, GlobalTransformArray);
-		for (int32 GeoIdx = 0; GeoIdx < Collection.NumElements(FGeometryCollection::GeometryGroup); ++GeoIdx)
+		if (FractureContext.GetGeometryCollection())
 		{
-			int32 TransformIdx = Collection.TransformIndex[GeoIdx];
-			FVector BoxCenter = GlobalTransformArray[TransformIdx].TransformPosition(Collection.BoundingBox[GeoIdx].GetCenter());
-			Vis.GeoCenters[GeoIdx] = BoxCenter;
-		}
+			FGeometryCollection& Collection = *FractureContext.GetGeometryCollection();
 
-		// Collect relevant proximity graph edges to render
-		auto AddEdgesForGeoIdx = [&Proximity, &Vis](int32 GeoIdx, bool bAddAll)
-		{
-			for (int32 GeoNbr : Proximity[GeoIdx])
+			if (!Collection.HasAttribute("Proximity", FGeometryCollection::GeometryGroup))
 			{
-				if (bAddAll || GeoIdx < GeoNbr)
-				{
-					Vis.ProximityEdges.Add(FEdgeVisInfo{ GeoIdx, GeoNbr });
-				}
+				continue;
 			}
-		};
-		if (ProximitySettings->bOnlyShowForSelected)
-		{
-			TArray<int32> SelectedRigids;
-			for (int32 TransformIdx : FractureContext.GetSelection())
+			const TManagedArray<TSet<int32>>& Proximity = Collection.GetAttribute<TSet<int32>>("Proximity", FGeometryCollection::GeometryGroup);
+
+			// Save the geometry collection component for rendering
+			int32 CollectionIdx = VisualizedCollections.Add(FractureContext.GetGeometryCollectionComponent());
+			FCollectionVisInfo& Vis = ProximityVisualizations.Emplace_GetRef();
+			Vis.CollectionIndex = CollectionIdx;
+
+			// Get proximity graph node positions as bounding box centers for each piece of geometry
+			Vis.GeoCenters.SetNum(Collection.NumElements(FGeometryCollection::GeometryGroup));
+			TArray<FTransform> GlobalTransformArray;
+			GeometryCollectionAlgo::GlobalMatrices(Collection.Transform, Collection.Parent, GlobalTransformArray);
+			for (int32 GeoIdx = 0; GeoIdx < Collection.NumElements(FGeometryCollection::GeometryGroup); ++GeoIdx)
+			{
+				int32 TransformIdx = Collection.TransformIndex[GeoIdx];
+				FVector BoxCenter = GlobalTransformArray[TransformIdx].TransformPosition(Collection.BoundingBox[GeoIdx].GetCenter());
+				Vis.GeoCenters[GeoIdx] = BoxCenter;
+			}
+
+			// Collect relevant proximity graph edges to render
+			auto AddEdgesForGeoIdx = [&Proximity, &Vis](int32 GeoIdx, bool bAddAll)
+				{
+					for (int32 GeoNbr : Proximity[GeoIdx])
+					{
+						if (bAddAll || GeoIdx < GeoNbr)
+						{
+							Vis.ProximityEdges.Add(FEdgeVisInfo{ GeoIdx, GeoNbr });
+						}
+					}
+				};
+			if (ProximitySettings->bOnlyShowForSelected)
 			{
-				SelectedRigids.Reset();
-				FGeometryCollectionClusteringUtility::GetLeafBones(&Collection, TransformIdx, true, SelectedRigids);
-				for (int32 RigidIdx : SelectedRigids)
+				TArray<int32> SelectedRigids;
+				for (int32 TransformIdx : FractureContext.GetSelection())
 				{
-					int32 GeoIdx = Collection.TransformToGeometryIndex[RigidIdx];
-					if (GeoIdx != INDEX_NONE)
+					SelectedRigids.Reset();
+					FGeometryCollectionClusteringUtility::GetLeafBones(&Collection, TransformIdx, true, SelectedRigids);
+					for (int32 RigidIdx : SelectedRigids)
 					{
-						AddEdgesForGeoIdx(GeoIdx, true);
+						int32 GeoIdx = Collection.TransformToGeometryIndex[RigidIdx];
+						if (GeoIdx != INDEX_NONE)
+						{
+							AddEdgesForGeoIdx(GeoIdx, true);
+						}
 					}
 				}
 			}
-		}
-		else
-		{
-			for (int32 GeoIdx = 0; GeoIdx < Collection.NumElements(FGeometryCollection::GeometryGroup); ++GeoIdx)
+			else
 			{
-				AddEdgesForGeoIdx(GeoIdx, false);
+				for (int32 GeoIdx = 0; GeoIdx < Collection.NumElements(FGeometryCollection::GeometryGroup); ++GeoIdx)
+				{
+					AddEdgesForGeoIdx(GeoIdx, false);
+				}
 			}
 		}
 	}
@@ -217,10 +220,13 @@ TArray<FFractureToolContext> UFractureToolProximity::GetFractureToolContexts() c
 
 	for (UGeometryCollectionComponent* GeometryCollectionComponent : GeomCompSelection)
 	{
-		// Generate a context for each selected node
-		FFractureToolContext FullSelection(GeometryCollectionComponent);
+		if (GeometryCollectionComponent->GetRestCollection())
+		{
+			// Generate a context for each selected node
+			FFractureToolContext FullSelection(GeometryCollectionComponent);
 
-		Contexts.Add(FullSelection);
+			Contexts.Add(FullSelection);
+		}
 	}
 
 	return Contexts;
diff --git a/Engine/Plugins/Experimental/ChaosEditor/Source/FractureEditor/Private/FractureToolUV.cpp b/Engine/Plugins/Experimental/ChaosEditor/Source/FractureEditor/Private/FractureToolUV.cpp
--- a/Engine/Plugins/Experimental/ChaosEditor/Source/FractureEditor/Private/FractureToolUV.cpp
+++ b/Engine/Plugins/Experimental/ChaosEditor/Source/FractureEditor/Private/FractureToolUV.cpp
@@ -165,8 +165,11 @@ void UFractureToolAutoUV::UpdateUVChannels(int32 TargetNumUVChannels)
 	GetSelectedGeometryCollectionComponents(GeomCompSelection);
 	for (UGeometryCollectionComponent* GeometryCollectionComponent : GeomCompSelection)
 	{
-		int32 NumChannels = GeometryCollectionComponent->GetRestCollection()->GetGeometryCollection()->NumUVLayers();
-		MinUVChannels = FMath::Min(MinUVChannels, NumChannels);
+		if (GeometryCollectionComponent->GetRestCollection())
+		{
+			const int32 NumChannels = GeometryCollectionComponent->GetRestCollection()->GetGeometryCollection()->NumUVLayers();
+			MinUVChannels = FMath::Min(MinUVChannels, NumChannels);
+		}
 	}
 
 	if (TargetNumUVChannels > -1 && TargetNumUVChannels != MinUVChannels)
@@ -176,27 +179,30 @@ void UFractureToolAutoUV::UpdateUVChannels(int32 TargetNumUVChannels)
 		bool bIsIncreasing = TargetNumUVChannels > MinUVChannels;
 		for (UGeometryCollectionComponent* GeometryCollectionComponent : GeomCompSelection)
 		{
-			int32 NumChannels = GeometryCollectionComponent->GetRestCollection()->GetGeometryCollection()->NumUVLayers();
-			// not if num channels is decreasing, all channels will be >= min channels so we apply the size change to all
-			// but if we're increasing the number of channels, some geometry collections may have more channels already and we leave those alone
-			if (NumChannels < TargetNumUVChannels || !bIsIncreasing)
+			if (GeometryCollectionComponent->GetRestCollection())
 			{
-				FGeometryCollectionEdit Edit(GeometryCollectionComponent, GeometryCollection::EEditUpdate::Rest, true);
-				GeometryCollectionComponent->GetRestCollection()->GetGeometryCollection()->SetNumUVLayers(TargetNumUVChannels);
-				if (bIsIncreasing)
+				const int32 NumChannels = GeometryCollectionComponent->GetRestCollection()->GetGeometryCollection()->NumUVLayers();
+				// not if num channels is decreasing, all channels will be >= min channels so we apply the size change to all
+				// but if we're increasing the number of channels, some geometry collections may have more channels already and we leave those alone
+				if (NumChannels < TargetNumUVChannels || !bIsIncreasing)
 				{
-					// Copy UV layer 0 into new UV layer
-					FGeometryCollection& Collection = *GeometryCollectionComponent->GetRestCollection()->GetGeometryCollection();
-					GeometryCollection::UV::FUVLayers Layers = GeometryCollection::UV::FindActiveUVLayers(Collection);
-					for (int32 Idx = 0; Idx < Layers[0].Num(); Idx++)
+					FGeometryCollectionEdit Edit(GeometryCollectionComponent, GeometryCollection::EEditUpdate::Rest, true);
+					GeometryCollectionComponent->GetRestCollection()->GetGeometryCollection()->SetNumUVLayers(TargetNumUVChannels);
+					if (bIsIncreasing)
 					{
-						for (int32 Ch = NumChannels; Ch < TargetNumUVChannels; Ch++)
+						// Copy UV layer 0 into new UV layer
+						FGeometryCollection& Collection = *GeometryCollectionComponent->GetRestCollection()->GetGeometryCollection();
+						GeometryCollection::UV::FUVLayers Layers = GeometryCollection::UV::FindActiveUVLayers(Collection);
+						for (int32 Idx = 0; Idx < Layers[0].Num(); Idx++)
 						{
-							Layers[Ch][Idx] = Layers[0][Idx];
+							for (int32 Ch = NumChannels; Ch < TargetNumUVChannels; Ch++)
+							{
+								Layers[Ch][Idx] = Layers[0][Idx];
+							}
 						}
 					}
+					Edit.GetRestCollection()->RebuildRenderData();
 				}
-				Edit.GetRestCollection()->RebuildRenderData();
 			}
 		}
 		MinUVChannels = TargetNumUVChannels;
@@ -213,9 +219,12 @@ void UFractureToolAutoUV::DisableBoneColors()
 	FScopedTransaction Transaction(LOCTEXT("DisableBoneColors", "Disable Bone Colors"), !GeomCompSelection.IsEmpty());
 	for (UGeometryCollectionComponent* GeometryCollectionComponent : GeomCompSelection)
 	{
-		GeometryCollectionComponent->Modify();
-		GeometryCollectionComponent->SetShowBoneColors(false);
-		GeometryCollectionComponent->MarkRenderStateDirty();
+		if (GeometryCollectionComponent->GetRestCollection())
+		{
+			GeometryCollectionComponent->Modify();
+			GeometryCollectionComponent->SetShowBoneColors(false);
+			GeometryCollectionComponent->MarkRenderStateDirty();
+		}
 	}
 }
 
@@ -232,16 +241,19 @@ void UFractureToolAutoUV::BoxProjectUVs()
 	FScopedTransaction Transaction(LOCTEXT("BoxProjectUVs", "Box Project UVs"), !GeomCompSelection.IsEmpty());
 	for (UGeometryCollectionComponent* GeometryCollectionComponent : GeomCompSelection)
 	{
-		FGeometryCollectionEdit Edit(GeometryCollectionComponent, GeometryCollection::EEditUpdate::Rest, true);
-		UE::PlanarCut::BoxProjectUVs(UVLayer,
-			*GeometryCollectionComponent->GetRestCollection()->GetGeometryCollection(),
-			(FVector3d)AutoUVSettings->ProjectionScale,
-			TargetFaces,
-			!TargetsMaterialIDs(AutoUVSettings->TargetFaces) ? EmptyMaterialIDs : AutoUVSettings->MaterialIDs,
-			(FVector2f)AutoUVSettings->ProjectionUVOffset,
-			AutoUVSettings->bAutoFitToBounds, AutoUVSettings->bCenterAtPivot, AutoUVSettings->bUniformProjectionScale);
-
-		Edit.GetRestCollection()->RebuildRenderData();
+		if (GeometryCollectionComponent->GetRestCollection())
+		{
+			FGeometryCollectionEdit Edit(GeometryCollectionComponent, GeometryCollection::EEditUpdate::Rest, true);
+			UE::PlanarCut::BoxProjectUVs(UVLayer,
+				*GeometryCollectionComponent->GetRestCollection()->GetGeometryCollection(),
+				(FVector3d)AutoUVSettings->ProjectionScale,
+				TargetFaces,
+				!TargetsMaterialIDs(AutoUVSettings->TargetFaces) ? EmptyMaterialIDs : AutoUVSettings->MaterialIDs,
+				(FVector2f)AutoUVSettings->ProjectionUVOffset,
+				AutoUVSettings->bAutoFitToBounds, AutoUVSettings->bCenterAtPivot, AutoUVSettings->bUniformProjectionScale);
+
+			Edit.GetRestCollection()->RebuildRenderData();
+		}
 	}
 }
 
@@ -264,31 +276,34 @@ TArray<FFractureToolContext> UFractureToolAutoUV::GetFractureToolContexts() cons
 
 	for (UGeometryCollectionComponent* GeometryCollectionComponent : GeomCompSelection)
 	{
-		// Generate a context for each selected node
-		FFractureToolContext FullSelection(GeometryCollectionComponent);
-		FullSelection.ConvertSelectionToRigidNodes();
-		
-		// Update global transforms and bounds -- TODO: pull this bounds update out to a shared function?
-		const TManagedArray<FTransform3f>& Transform = FullSelection.GetGeometryCollection()->Transform;
-		const TManagedArray<int32>& TransformToGeometryIndex = FullSelection.GetGeometryCollection()->TransformToGeometryIndex;
-		const TManagedArray<FBox>& BoundingBoxes = FullSelection.GetGeometryCollection()->BoundingBox;
+		if (GeometryCollectionComponent->GetRestCollection())
+		{
+			// Generate a context for each selected node
+			FFractureToolContext FullSelection(GeometryCollectionComponent);
+			FullSelection.ConvertSelectionToRigidNodes();
 
-		TArray<FTransform> Transforms;
-		GeometryCollectionAlgo::GlobalMatrices(Transform, FullSelection.GetGeometryCollection()->Parent, Transforms);
+			// Update global transforms and bounds -- TODO: pull this bounds update out to a shared function?
+			const TManagedArray<FTransform3f>& Transform = FullSelection.GetGeometryCollection()->Transform;
+			const TManagedArray<int32>& TransformToGeometryIndex = FullSelection.GetGeometryCollection()->TransformToGeometryIndex;
+			const TManagedArray<FBox>& BoundingBoxes = FullSelection.GetGeometryCollection()->BoundingBox;
 
-		FBox Bounds(ForceInit);
-		for (int32 BoneIndex : FullSelection.GetSelection())
-		{
-			int32 GeometryIndex = TransformToGeometryIndex[BoneIndex];
-			if (GeometryIndex > INDEX_NONE)
+			TArray<FTransform> Transforms;
+			GeometryCollectionAlgo::GlobalMatrices(Transform, FullSelection.GetGeometryCollection()->Parent, Transforms);
+
+			FBox Bounds(ForceInit);
+			for (int32 BoneIndex : FullSelection.GetSelection())
 			{
-				FBox BoneBound = BoundingBoxes[GeometryIndex].TransformBy(Transforms[BoneIndex]);
-				Bounds += BoneBound;
+				int32 GeometryIndex = TransformToGeometryIndex[BoneIndex];
+				if (GeometryIndex > INDEX_NONE)
+				{
+					FBox BoneBound = BoundingBoxes[GeometryIndex].TransformBy(Transforms[BoneIndex]);
+					Bounds += BoneBound;
+				}
 			}
-		}
-		FullSelection.SetBounds(Bounds);
+			FullSelection.SetBounds(Bounds);
 
-		Contexts.Add(FullSelection);
+			Contexts.Add(FullSelection);
+		}
 	}
 
 	return Contexts;
@@ -452,11 +467,14 @@ void UFractureToolAutoUV::LayoutUVs()
 	FScopedTransaction Transaction(LOCTEXT("LayoutUVs", "Layout UVs"), !GeomCompSelection.IsEmpty());
 	for (UGeometryCollectionComponent* GeometryCollectionComponent : GeomCompSelection)
 	{
-		FGeometryCollectionEdit Edit(GeometryCollectionComponent, GeometryCollection::EEditUpdate::None);
-		Edit.GetRestCollection()->Modify();
-		LayoutUVsForComponent(GeometryCollectionComponent);
+		if (GeometryCollectionComponent->GetRestCollection())
+		{
+			FGeometryCollectionEdit Edit(GeometryCollectionComponent, GeometryCollection::EEditUpdate::None);
+			Edit.GetRestCollection()->Modify();
+			LayoutUVsForComponent(GeometryCollectionComponent);
 
-		Edit.GetRestCollection()->RebuildRenderData();
+			Edit.GetRestCollection()->RebuildRenderData();
+		}
 	}
 }
 
@@ -523,7 +541,10 @@ void UFractureToolAutoUV::BakeTexture()
 	FScopedTransaction Transaction(LOCTEXT("BakeTexture", "Bake Texture"), !GeomCompSelection.IsEmpty());
 	for (UGeometryCollectionComponent* GeometryCollectionComponent : GeomCompSelection)
 	{
-		BakeTextureForComponent(GeometryCollectionComponent);
+		if (GeometryCollectionComponent->GetRestCollection())
+		{
+			BakeTextureForComponent(GeometryCollectionComponent);
+		}
 	}
 }
 
diff --git a/Engine/Plugins/Experimental/ChaosEditor/Source/FractureEditor/Public/FractureTool.cpp b/Engine/Plugins/Experimental/ChaosEditor/Source/FractureEditor/Public/FractureTool.cpp
--- a/Engine/Plugins/Experimental/ChaosEditor/Source/FractureEditor/Public/FractureTool.cpp
+++ b/Engine/Plugins/Experimental/ChaosEditor/Source/FractureEditor/Public/FractureTool.cpp
@@ -277,7 +277,10 @@ void UFractureActionTool::SetOutlinerComponents(TArray<FFractureToolContext>& In
 	TArray<UGeometryCollectionComponent*> Components;
 	for (FFractureToolContext& Context : InContexts)
 	{
-		Components.AddUnique(Context.GetGeometryCollectionComponent());
+		if (Context.GetGeometryCollection())
+		{
+			Components.AddUnique(Context.GetGeometryCollectionComponent());
+		}
 	}
 	Toolkit->SetOutlinerComponents(Components);
 }
@@ -378,44 +381,47 @@ void UFractureModalTool::Execute(TWeakPtr<FFractureEditorModeToolkit> InToolkit)
 
 		for (FFractureToolContext& FractureContext : FractureContexts)
 		{
-			FGeometryCollectionEdit EditCollection(FractureContext.GetGeometryCollectionComponent(), GeometryCollection::EEditUpdate::RestPhysicsDynamic, !ExecuteUpdatesShape());
-
-			int32 InitialNumTransforms = FractureContext.GetGeometryCollection()->NumElements(FGeometryCollection::TransformGroup);
-			int32 FirstNewGeometryIndex = ExecuteFracture(FractureContext);
-			
-			if (FirstNewGeometryIndex > INDEX_NONE)
+			if (FractureContext.GetGeometryCollection())
 			{
-				PostFractureProcess(FractureContext, FirstNewGeometryIndex);
-
-				FractureContext.GenerateGuids(FirstNewGeometryIndex);
+				FGeometryCollectionEdit EditCollection(FractureContext.GetGeometryCollectionComponent(), GeometryCollection::EEditUpdate::RestPhysicsDynamic, !ExecuteUpdatesShape());
 
-				// Based on the first new geometry index, generate a list of new transforms generated by the fracture.
-				const TManagedArray<int32>& TransformIndex = FractureContext.GetGeometryCollection()->GetAttribute<int32>("TransformIndex", FGeometryCollection::GeometryGroup);
-				int32 LastGeometryIndex = TransformIndex.Num();
+				int32 InitialNumTransforms = FractureContext.GetGeometryCollection()->NumElements(FGeometryCollection::TransformGroup);
+				int32 FirstNewGeometryIndex = ExecuteFracture(FractureContext);
 
-				TArray<int32> NewTransforms;
-				NewTransforms.Reserve(LastGeometryIndex - FirstNewGeometryIndex);
-				for (int32 GeometryIndex = FirstNewGeometryIndex; GeometryIndex < LastGeometryIndex; ++GeometryIndex)
+				if (FirstNewGeometryIndex > INDEX_NONE)
 				{
-					NewTransforms.Add(TransformIndex[GeometryIndex]);
-				}
+					PostFractureProcess(FractureContext, FirstNewGeometryIndex);
 
-				FractureContext.SetSelection(NewTransforms);
-			}
-			else
-			{
-				// either no update was done, or the updated range was not expressible as a final geometry index
-				// -- in the latter case, selection may become invalid, so we should clear it
-				int32 NumTransforms = FractureContext.GetGeometryCollection()->NumElements(FGeometryCollection::TransformGroup);
-				TArray<int32>& Selection = FractureContext.GetSelection();
-				if (InitialNumTransforms != NumTransforms)
+					FractureContext.GenerateGuids(FirstNewGeometryIndex);
+
+					// Based on the first new geometry index, generate a list of new transforms generated by the fracture.
+					const TManagedArray<int32>& TransformIndex = FractureContext.GetGeometryCollection()->GetAttribute<int32>("TransformIndex", FGeometryCollection::GeometryGroup);
+					int32 LastGeometryIndex = TransformIndex.Num();
+
+					TArray<int32> NewTransforms;
+					NewTransforms.Reserve(LastGeometryIndex - FirstNewGeometryIndex);
+					for (int32 GeometryIndex = FirstNewGeometryIndex; GeometryIndex < LastGeometryIndex; ++GeometryIndex)
+					{
+						NewTransforms.Add(TransformIndex[GeometryIndex]);
+					}
+
+					FractureContext.SetSelection(NewTransforms);
+				}
+				else
 				{
-					Selection.Empty();
+					// either no update was done, or the updated range was not expressible as a final geometry index
+					// -- in the latter case, selection may become invalid, so we should clear it
+					int32 NumTransforms = FractureContext.GetGeometryCollection()->NumElements(FGeometryCollection::TransformGroup);
+					TArray<int32>& Selection = FractureContext.GetSelection();
+					if (InitialNumTransforms != NumTransforms)
+					{
+						Selection.Empty();
+					}
+					FractureContext.Sanitize(false);
 				}
-				FractureContext.Sanitize(false);
-			}
 
-			Refresh(FractureContext, Toolkit);
+				Refresh(FractureContext, Toolkit);
+			}
 		}
 
 		SetOutlinerComponents(FractureContexts, Toolkit);
