fix fracture mode stats not refreshing when hull count changes by adding hull count to cached stats

diff --git a/Engine/Plugins/Experimental/ChaosEditor/Source/FractureEditor/Private/FractureEditorModeToolkit.cpp b/Engine/Plugins/Experimental/ChaosEditor/Source/FractureEditor/Private/FractureEditorModeToolkit.cpp
--- a/Engine/Plugins/Experimental/ChaosEditor/Source/FractureEditor/Private/FractureEditorModeToolkit.cpp
+++ b/Engine/Plugins/Experimental/ChaosEditor/Source/FractureEditor/Private/FractureEditorModeToolkit.cpp
@@ -1831,7 +1831,7 @@ FText FFractureEditorModeToolkit::GetActiveToolMessage() const
 
 namespace UE::PrivateOutlinerCacheHelper
 {
-	void AddGeometryStatsForComponent(const UGeometryCollectionComponent* Component, int64& BoneCount, int64& VertexCount)
+	void AddGeometryStatsForComponent(const UGeometryCollectionComponent* Component, int64& BoneCount, int64& VertexCount, int64& HullCount)
 	{
 		if (Component)
 		{
@@ -1841,26 +1841,27 @@ namespace UE::PrivateOutlinerCacheHelper
 				{
 					BoneCount += Collection->Transform.Num();
 					VertexCount += Collection->Vertex.Num();
+					HullCount += Collection->NumElements(FGeometryCollection::ConvexGroup);
 				}
 			}
 		}
 	}
 
-	void ComputeGeometryStats(const TArray<UGeometryCollectionComponent*>& InNewComponents, int64& BoneCount, int64& VertexCount)
+	void ComputeGeometryStats(const TArray<UGeometryCollectionComponent*>& InNewComponents, int64& BoneCount, int64& VertexCount, int64& HullCount)
 	{
-		BoneCount = VertexCount = 0;
+		BoneCount = VertexCount = HullCount = 0;
 		for (const UGeometryCollectionComponent* Component : InNewComponents)
 		{
-			AddGeometryStatsForComponent(Component, BoneCount, VertexCount);
+			AddGeometryStatsForComponent(Component, BoneCount, VertexCount, HullCount);
 		}
 	}
 
-	void ComputeGeometryStats(const TArray<TWeakObjectPtr<UGeometryCollectionComponent>>& InNewComponents, int64& BoneCount, int64& VertexCount)
+	void ComputeGeometryStats(const TArray<TWeakObjectPtr<UGeometryCollectionComponent>>& InNewComponents, int64& BoneCount, int64& VertexCount, int64& HullCount)
 	{
-		BoneCount = VertexCount = 0;
+		BoneCount = VertexCount = HullCount = 0;
 		for (const TWeakObjectPtr<UGeometryCollectionComponent> Component : InNewComponents)
 		{
-			AddGeometryStatsForComponent(Component.Get(), BoneCount, VertexCount);
+			AddGeometryStatsForComponent(Component.Get(), BoneCount, VertexCount, HullCount);
 		}
 	}
 }
@@ -1868,15 +1869,15 @@ namespace UE::PrivateOutlinerCacheHelper
 bool FFractureEditorModeToolkit::IsCachedOutlinerGeometryStale(const TArray<TWeakObjectPtr<UGeometryCollectionComponent>>& SelectedComponents) const
 {
 	// Note we currently use quick-to-compute high level stats, compared vs cached versions, since this is run per tick; we could change this to more thoroughly walk the outliner data potentially
-	int64 NewBoneCount, NewVertexCount;
-	UE::PrivateOutlinerCacheHelper::ComputeGeometryStats(SelectedComponents, NewBoneCount, NewVertexCount);
-	return NewBoneCount != OutlinerCachedBoneCount || NewVertexCount != OutlinerCachedVertexCount;
+	int64 NewBoneCount, NewVertexCount, NewHullCount;
+	UE::PrivateOutlinerCacheHelper::ComputeGeometryStats(SelectedComponents, NewBoneCount, NewVertexCount, NewHullCount);
+	return NewBoneCount != OutlinerCachedBoneCount || NewVertexCount != OutlinerCachedVertexCount || NewHullCount != OutlinerCachedHullCount;
 }
 
 void FFractureEditorModeToolkit::SetOutlinerComponents(const TArray<UGeometryCollectionComponent*>& InNewComponents)
 {
 	// Update cached stats (bone and vertex count) of the components in the outliner
-	UE::PrivateOutlinerCacheHelper::ComputeGeometryStats(InNewComponents, OutlinerCachedBoneCount, OutlinerCachedVertexCount);
+	UE::PrivateOutlinerCacheHelper::ComputeGeometryStats(InNewComponents, OutlinerCachedBoneCount, OutlinerCachedVertexCount, OutlinerCachedHullCount);
 
 	TArray<UGeometryCollectionComponent*> ComponentsToEdit;
 	ComponentsToEdit.Reserve(InNewComponents.Num());
diff --git a/Engine/Plugins/Experimental/ChaosEditor/Source/FractureEditor/Public/FractureEditorModeToolkit.h b/Engine/Plugins/Experimental/ChaosEditor/Source/FractureEditor/Public/FractureEditorModeToolkit.h
--- a/Engine/Plugins/Experimental/ChaosEditor/Source/FractureEditor/Public/FractureEditorModeToolkit.h
+++ b/Engine/Plugins/Experimental/ChaosEditor/Source/FractureEditor/Public/FractureEditorModeToolkit.h
@@ -275,6 +275,5 @@ private:
 	TSharedPtr<STextComboBox> AssetLocationMode;
 
 	// Simple cached statistics to allow use to quickly/heuristically check for stale geometry collection data
-	int64 OutlinerCachedBoneCount = 0, OutlinerCachedVertexCount = 0;
-
+	int64 OutlinerCachedBoneCount = 0, OutlinerCachedVertexCount = 0, OutlinerCachedHullCount = 0;
 };
