add guards against stale selection crashing the fracture editor mode when trying to render selection boxes

diff --git a/Engine/Plugins/Experimental/ChaosEditor/Source/FractureEditor/Private/FractureEditorMode.cpp b/Engine/Plugins/Experimental/ChaosEditor/Source/FractureEditor/Private/FractureEditorMode.cpp
--- a/Engine/Plugins/Experimental/ChaosEditor/Source/FractureEditor/Private/FractureEditorMode.cpp
+++ b/Engine/Plugins/Experimental/ChaosEditor/Source/FractureEditor/Private/FractureEditorMode.cpp
@@ -246,16 +246,23 @@ void UFractureEditorMode::RenderSelectionAsBoxes(const FSceneView* View, FViewpo
 			const FTransform& WorldTransform = SelectedComp->GetComponentToWorld();
 
 			const TManagedArray<FVector3f>* ExplodedVectors = nullptr;
+			int32 NumTransforms = 0;
 			if (SelectedComp->GetRestCollection())
 			{
 				if (TSharedPtr<const FGeometryCollection> Collection = SelectedComp->GetRestCollection()->GetGeometryCollection())
 				{
 					ExplodedVectors = Collection->FindAttribute<FVector3f>("ExplodedVector", FGeometryCollection::TransformGroup);
+					NumTransforms = Collection->NumElements(FGeometryCollection::TransformGroup);
 				}
 			}
 
 			for (int32 SelectedBone : SelectedComp->GetSelectedBones())
 			{
+				if (SelectedBone < 0 || SelectedBone >= NumTransforms)
+				{
+					continue;
+				}
+
 				float BoneMass = 0;
 				FBox BoneExtent;
 				SelectedComp->GetMassAndExtents(SelectedBone, BoneMass, BoneExtent);
diff --git a/Engine/Source/Runtime/Experimental/GeometryCollectionEngine/Private/GeometryCollection/GeometryCollectionComponent.cpp b/Engine/Source/Runtime/Experimental/GeometryCollectionEngine/Private/GeometryCollection/GeometryCollectionComponent.cpp
--- a/Engine/Source/Runtime/Experimental/GeometryCollectionEngine/Private/GeometryCollection/GeometryCollectionComponent.cpp
+++ b/Engine/Source/Runtime/Experimental/GeometryCollectionEngine/Private/GeometryCollection/GeometryCollectionComponent.cpp
@@ -6978,6 +6978,10 @@ void UGeometryCollectionComponent::GetMassAndExtents(int32 ItemIndex, float& Out
 			const TManagedArray<FBox>* GeoBoundingBoxes = Collection.FindAttribute<FBox>(BoundingBoxAttributeName, FGeometryCollection::GeometryGroup);
 
 			const FGeometryCollectionItemIndex GCItemIndex = FGeometryCollectionItemIndex::CreateFromExistingItemIndex(ItemIndex);
+			if (!ensure(CollectionMass->IsValidIndex(GCItemIndex.GetTransformIndex())))
+			{
+				return;
+			}
 
 			if (GCItemIndex.IsInternalCluster())
 			{
