Dataflow : fix fracture tool mode box selection  not properly accounting for transforms and clusters
- Also add a selection display mode in the fracture setting sto choose between highlight and bounding boxes

diff --git a/Engine/Plugins/Experimental/ChaosEditor/Source/FractureEditor/Private/FractureEditorMode.cpp b/Engine/Plugins/Experimental/ChaosEditor/Source/FractureEditor/Private/FractureEditorMode.cpp
--- a/Engine/Plugins/Experimental/ChaosEditor/Source/FractureEditor/Private/FractureEditorMode.cpp
+++ b/Engine/Plugins/Experimental/ChaosEditor/Source/FractureEditor/Private/FractureEditorMode.cpp
@@ -21,6 +21,8 @@
 #include "GeometryCollection/GeometryCollection.h"
 #include "GeometryCollection/GeometryCollectionUtility.h"
 #include "GeometryCollection/GeometryCollectionClusteringUtility.h"
+#include "GeometryCollection/Facades/CollectionTransformSelectionFacade.h"
+#include "GeometryCollection/Facades/CollectionBoundsFacade.h"
 #include "FractureSelectionTools.h"
 #include "EditorViewportClient.h"
 #include "Misc/ITransaction.h"
@@ -31,7 +33,9 @@
 #include "BaseGizmos/TransformGizmoUtil.h"
 #include "Snapping/ModelingSceneSnappingManager.h"
 #include "Components/InstancedStaticMeshComponent.h"
+#include "PrimitiveSceneProxy.h"
 #include "SLevelViewport.h"
+#include "FractureSettings.h"
 
 #include "UnrealEdGlobals.h"
 #include "EditorModeManager.h"
@@ -224,6 +228,23 @@ void UFractureEditorMode::Render(const FSceneView* View, FViewport* Viewport, FP
 	}
 
 	RenderSelectionAsBoxes(View, Viewport, PDI);
+
+	// if highlight is not selected disable them
+	// technically we should prevent the highlight to be set originally but too many calling site, will need some refactor
+	if (const UFractureSettings* FractureSettings = GetDefault<UFractureSettings>())
+	{
+		if (FractureSettings->SelectionDisplayMode != EFractureSelectionDisplayMode::Highlight)
+		{
+			for (TWeakObjectPtr<UGeometryCollectionComponent> SelectedCompWeakPtr : SelectedGeometryComponents)
+			{
+				if (UGeometryCollectionComponent* SelectedComp = SelectedCompWeakPtr.Get())
+				{
+					FScopedColorEdit EditBoneColor = SelectedComp->EditBoneSelection();
+					EditBoneColor.ResetHighlightedBones();
+				}
+			}
+		}
+	}
 }
 
 void UFractureEditorMode::DrawHUD(FEditorViewportClient* ViewportClient, FViewport* Viewport, const FSceneView* View, FCanvas* Canvas)
@@ -239,23 +260,43 @@ void UFractureEditorMode::DrawHUD(FEditorViewportClient* ViewportClient, FViewpo
 
 void UFractureEditorMode::RenderSelectionAsBoxes(const FSceneView* View, FViewport* Viewport, FPrimitiveDrawInterface* PDI)
 {
+	// this is a uggestion has nanite based geometry collection may want to override it 
+	bool bShouldDisplayBounds = false;
+	if (const UFractureSettings* FractureSettings = GetDefault<UFractureSettings>())
+	{
+		bShouldDisplayBounds = (FractureSettings->SelectionDisplayMode == EFractureSelectionDisplayMode::BoundingBox);
+	}
+
 	for (TWeakObjectPtr<UGeometryCollectionComponent> SelectedCompWeakPtr : SelectedGeometryComponents)
 	{
 		if (UGeometryCollectionComponent* SelectedComp = SelectedCompWeakPtr.Get())
 		{
-			const FTransform& WorldTransform = SelectedComp->GetComponentToWorld();
+			TSharedPtr<const FGeometryCollection> Collection = (SelectedComp->GetRestCollection()) ? SelectedComp->GetRestCollection()->GetGeometryCollection() : nullptr;
+			if (!Collection)
+			{
+				continue;
+			}
 
-			const TManagedArray<FVector3f>* ExplodedVectors = nullptr;
-			int32 NumTransforms = 0;
-			if (SelectedComp->GetRestCollection())
+			const bool bUsesNanite = SelectedComp->GetSceneProxy() && SelectedComp->GetSceneProxy()->IsNaniteMesh();
+			if (!bShouldDisplayBounds && !bUsesNanite)
 			{
-				if (TSharedPtr<const FGeometryCollection> Collection = SelectedComp->GetRestCollection()->GetGeometryCollection())
-				{
-					ExplodedVectors = Collection->FindAttribute<FVector3f>("ExplodedVector", FGeometryCollection::TransformGroup);
-					NumTransforms = Collection->NumElements(FGeometryCollection::TransformGroup);
-				}
+				continue;
+			}
+
+			const GeometryCollection::Facades::FCollectionTransformSelectionFacade TransformFacade(*Collection);
+			const GeometryCollection::Facades::FBoundsFacade BoundsFacade(*Collection);
+			if (!TransformFacade.IsValid() || !BoundsFacade.IsValid())
+			{
+				continue;
 			}
 
+			const FTransform& WorldTransform = SelectedComp->GetComponentToWorld();
+			const TArray<FTransform3f>& ComponentSpaceTransforms = SelectedComp->GetComponentSpaceTransforms3f();
+
+			const int32 NumTransforms = Collection->NumElements(FGeometryCollection::TransformGroup);
+			const TManagedArray<FBox>& GeometryBounds = BoundsFacade.GetBoundingBoxes();
+			const TManagedArray<int32>& TransformToGeometry = BoundsFacade.GetTransformToGeometryIndex();
+
 			for (int32 SelectedBone : SelectedComp->GetSelectedBones())
 			{
 				if (SelectedBone < 0 || SelectedBone >= NumTransforms)
@@ -263,12 +304,26 @@ void UFractureEditorMode::RenderSelectionAsBoxes(const FSceneView* View, FViewpo
 					continue;
 				}
 
-				float BoneMass = 0;
 				FBox BoneExtent;
-				SelectedComp->GetMassAndExtents(SelectedBone, BoneMass, BoneExtent);
 
-				const FVector3f Offset = ExplodedVectors ? (*ExplodedVectors)[SelectedBone] : FVector3f::ZeroVector;
-				const FVector BoxCenter = BoneExtent.GetCenter() + FVector(Offset);
+				// only leaf bones have geometry bounds we need to aggregate them 
+				TArray<int32> LeafBoneIndices;
+				TransformFacade.ConvertSelectionToRigidNodes(SelectedBone, LeafBoneIndices);
+				for (const int32 LeafBoneIndex : LeafBoneIndices)
+				{
+					const int32 GeometryIndex = TransformToGeometry[LeafBoneIndex];
+					if (GeometryBounds.IsValidIndex(GeometryIndex) && ComponentSpaceTransforms.IsValidIndex(LeafBoneIndex))
+					{
+						BoneExtent += GeometryBounds[GeometryIndex].TransformBy(FTransform(ComponentSpaceTransforms[LeafBoneIndex]));
+					}
+				}
+
+				if (!BoneExtent.IsValid)
+				{
+					continue;
+				}
+
+				const FVector BoxCenter = BoneExtent.GetCenter();// +FVector(Offset);
 				const FVector BoxHalfSize = BoneExtent.GetExtent();
 				const FVector Corners[] =
 				{
diff --git a/Engine/Plugins/Experimental/ChaosEditor/Source/FractureEditor/Private/FractureSettings.cpp b/Engine/Plugins/Experimental/ChaosEditor/Source/FractureEditor/Private/FractureSettings.cpp
--- a/Engine/Plugins/Experimental/ChaosEditor/Source/FractureEditor/Private/FractureSettings.cpp
+++ b/Engine/Plugins/Experimental/ChaosEditor/Source/FractureEditor/Private/FractureSettings.cpp
@@ -11,6 +11,7 @@ UFractureSettings::UFractureSettings(const FObjectInitializer& ObjInit)
 	, ExplodeAmount(0.0)
 	, FractureLevel(-1)
 	, bHideUnselected(false)
+	, SelectionDisplayMode(EFractureSelectionDisplayMode::Highlight)
 {}
 
 #endif
diff --git a/Engine/Plugins/Experimental/ChaosEditor/Source/FractureEditor/Public/FractureSettings.h b/Engine/Plugins/Experimental/ChaosEditor/Source/FractureEditor/Public/FractureSettings.h
--- a/Engine/Plugins/Experimental/ChaosEditor/Source/FractureEditor/Public/FractureSettings.h
+++ b/Engine/Plugins/Experimental/ChaosEditor/Source/FractureEditor/Public/FractureSettings.h
@@ -9,6 +9,17 @@
 
 class UGeometryCollection;
 
+UENUM()
+enum class EFractureSelectionDisplayMode : uint8
+{
+	/** Use a glow material to highlight the selection - Nanite Geometry Collection will ignore it and use bounding boxes instead */
+	Highlight,
+
+	/** Use bounding boxes to show selection - Nanite Geometry Collection will default to that */
+	BoundingBox,
+};
+
+
 /** Settings specifically related to viewing fractured meshes **/
 UCLASS()
 class UFractureSettings: public UObject
@@ -30,6 +41,10 @@ public:
 	UPROPERTY(EditAnywhere, Category = ViewSettings, meta = (DisplayName = "Hide Unselected"))
 	bool bHideUnselected;
 
+	/** How to display the selection in the viewport */
+	UPROPERTY(EditAnywhere, Category = ViewSettings, meta = (DisplayName = "Highlight Selected"))
+	EFractureSelectionDisplayMode SelectionDisplayMode;
+
 	UPROPERTY(VisibleAnywhere, Category = ViewSettings)
 	TWeakObjectPtr<const UGeometryCollection> RestCollection = nullptr;
 
