fix fracture tool resample point vis being too hard to see (made bigger + magenta)

diff --git a/Engine/Plugins/Experimental/ChaosEditor/Source/FractureEditor/Private/FractureToolResample.cpp b/Engine/Plugins/Experimental/ChaosEditor/Source/FractureEditor/Private/FractureToolResample.cpp
--- a/Engine/Plugins/Experimental/ChaosEditor/Source/FractureEditor/Private/FractureToolResample.cpp
+++ b/Engine/Plugins/Experimental/ChaosEditor/Source/FractureEditor/Private/FractureToolResample.cpp
@@ -47,7 +47,7 @@ void UFractureToolResample::Render(const FSceneView* View, FViewport* Viewport,
 	EnumerateVisualizationMapping(PointsMappings, GeneratedPoints.Num(), [&](int32 Idx, FVector ExplodedVector)
 	{
 		FVector Point = GeneratedPoints[Idx];
-		PDI->DrawPoint(Point + ExplodedVector, FLinearColor::Green, 2.f, SDPG_Foreground);
+		PDI->DrawPoint(Point + ExplodedVector, FLinearColor(1,0,1), 4.f, SDPG_Foreground);
 	});
 }
 
