Expanded out fracture brick pattern bounds by a half-brick, to fix bricks not quite covering the intended bounds in many cases

diff --git a/Engine/Plugins/Experimental/ChaosEditor/Source/FractureEditor/Private/FractureToolBrick.cpp b/Engine/Plugins/Experimental/ChaosEditor/Source/FractureEditor/Private/FractureToolBrick.cpp
--- a/Engine/Plugins/Experimental/ChaosEditor/Source/FractureEditor/Private/FractureToolBrick.cpp
+++ b/Engine/Plugins/Experimental/ChaosEditor/Source/FractureEditor/Private/FractureToolBrick.cpp
@@ -7,6 +7,7 @@
 #include "Editor.h"
 #include "Engine/Selection.h"
 #include "FractureEditorStyle.h"
+#include "FractureEngineFracturing.h"
 #include "GameFramework/Actor.h"
 #include "GeometryCollection/GeometryCollection.h"
 #include "GeometryCollection/GeometryCollectionObject.h"
@@ -153,154 +154,11 @@ namespace FractureToolBrickLocals
 
 void UFractureToolBrick::GenerateBrickTransforms(const FBox& Bounds)
 {
-	const FVector& Min = Bounds.Min;
-	const FVector& Max = Bounds.Max;
-	const FVector Extents(Bounds.Max - Bounds.Min);
-
-	// Determine brick dimensions (length, depth, height) and make sure we do not exceed the limit for the number of bricks.
-	// If we would simply use the input dimensions, we are prone to running out of memory and/or exceeding the storage capabilities of TArray, and crash.
-	const FVector BrickDimensions = FractureToolBrickLocals::GetBrickDimensions(BrickSettings.Get(), Extents);
-
-	// Early out if we have inputs we cannot deal with.
-	if (BrickDimensions == FVector::ZeroVector)
-	{
-		return;
-	}
-
-	// Reserve correct amount of memory to avoid re-allocations.
-	BrickTransforms.Reserve(FractureToolBrickLocals::CalculateNumBricks(BrickDimensions, Extents));
-
-	const FVector BrickHalfDimensions(BrickDimensions * 0.5);
-	const FQuat HeaderRotation(FVector::UpVector, 1.5708);
-
-	if (BrickSettings->Bond == EFractureBrickBond::Stretcher)
-	{
-		bool OddY = false;
-		for (float yy = 0.f; yy <= Extents.Y; yy += BrickDimensions.Y)
-		{
-			bool Oddline = false;
-			for (float zz = BrickHalfDimensions.Z; zz <= Extents.Z; zz += BrickDimensions.Z)
-			{
-				for (float xx = 0.f; xx <= Extents.X; xx += BrickDimensions.X)
-				{
-					FVector BrickPosition(Min + FVector(Oddline ^ OddY ? xx : xx + BrickHalfDimensions.X, yy, zz));
-					BrickTransforms.Emplace(FTransform(BrickPosition));
-				}
-				Oddline = !Oddline;
-			}
-			OddY = !OddY;
-		}
-	}
-	else if (BrickSettings->Bond == EFractureBrickBond::Stack)
-	{
-		bool OddY = false;
-		for (float yy = 0.f; yy <= Extents.Y; yy += BrickDimensions.Y)
-		{
-			for (float zz = BrickHalfDimensions.Z; zz <= Extents.Z; zz += BrickDimensions.Z)
-			{
-				for (float xx = 0.f; xx <= Extents.X; xx += BrickDimensions.X)
-				{
-					FVector BrickPosition(Min + FVector(OddY ? xx : xx + BrickHalfDimensions.X, yy, zz));
-					BrickTransforms.Emplace(FTransform(BrickPosition));
-				}
-			}
-			OddY = !OddY;
-		}
-	}
-	else if (BrickSettings->Bond == EFractureBrickBond::English)
-	{
-		float HalfLengthDepthDifference = BrickHalfDimensions.X - BrickHalfDimensions.Y - BrickHalfDimensions.Y;
-		bool OddY = false;
-		for (float yy = 0.f; yy <= Extents.Y; yy += BrickDimensions.Y)
-		{
-			bool Oddline = false;
-			for (float zz = BrickHalfDimensions.Z; zz <= Extents.Z; zz += BrickDimensions.Z)
-			{
-				if (Oddline && !OddY) // header row
-				{
-					for (float xx = 0.f; xx <= Extents.X; xx += BrickDimensions.Y)
-					{
-						FVector BrickPosition(Min + FVector(Oddline ^ OddY ? xx : xx + BrickHalfDimensions.Y, yy + BrickHalfDimensions.Y, zz));
-						BrickTransforms.Emplace(FTransform(HeaderRotation, BrickPosition));
-					}
-				}
-				else if(!Oddline) // stretchers
-				{
-					for (float xx = 0.f; xx <= Extents.X; xx += BrickDimensions.X)
-					{
-						FVector BrickPosition(Min + FVector(Oddline ^ OddY ? xx : xx + BrickHalfDimensions.X, OddY ? yy + HalfLengthDepthDifference : yy - HalfLengthDepthDifference , zz));
-						BrickTransforms.Emplace(FTransform(BrickPosition));
-					}
-				}
-				Oddline = !Oddline;
-			}
-			OddY = !OddY;
-		}
-	}
-	else if (BrickSettings->Bond == EFractureBrickBond::Header)
-	{
-		bool OddY = false;
-		for (float yy = 0.f; yy <= Extents.Y; yy += BrickDimensions.X)
-		{
-			bool Oddline = false;
-			for (float zz = BrickHalfDimensions.Z; zz <= Extents.Z; zz += BrickDimensions.Z)
-			{
-				for (float xx = 0.f; xx <= Extents.X; xx += BrickDimensions.Y)
-				{
-					FVector BrickPosition(Min + FVector(Oddline ^ OddY ? xx : xx + BrickHalfDimensions.Y, yy, zz));
-					BrickTransforms.Emplace(FTransform(HeaderRotation, BrickPosition));
-				}
-				Oddline = !Oddline;
-			}
-			OddY = !OddY;
-		}
-	}
-	else if (BrickSettings->Bond == EFractureBrickBond::Flemish)
-	{
-		float HalfLengthDepthDifference = BrickHalfDimensions.X - BrickDimensions.Y;
-		bool OddY = false;
-		int32 RowY = 0;
-		for (float yy = 0.f; yy <= Extents.Y; yy += BrickDimensions.Y)
-		{
-			bool OddZ = false;
-			for (float zz = BrickHalfDimensions.Z; zz <= Extents.Z; zz += BrickDimensions.Z)
-			{
-				bool OddX = OddZ;
-				for (float xx = 0.f; xx <= Extents.X; xx += BrickHalfDimensions.X + BrickHalfDimensions.Y)
-				{
-					FVector BrickPosition(Min + FVector(xx,yy,zz));
-					if (OddX)
-					{
-						if(OddY) // runner
-						{
-							BrickTransforms.Emplace(FTransform(BrickPosition + FVector(0, HalfLengthDepthDifference, 0))); // runner
-						}
-						else
-						{
-							BrickTransforms.Emplace(FTransform(BrickPosition - FVector(0, HalfLengthDepthDifference, 0))); // runner
-
-						}
-					}
-					else if (!OddY) // header
-					{
-						BrickTransforms.Emplace(FTransform(HeaderRotation, BrickPosition + FVector(0, BrickHalfDimensions.Y, 0))); // header
-					}
-					OddX = !OddX;
-				}
-				OddZ = !OddZ;
-			}
-			OddY = !OddY;
-			++RowY;
-		}
-	}
-
-	const FVector BrickMax(BrickHalfDimensions);
-	const FVector BrickMin(-BrickHalfDimensions);
-
-	for (const auto& Transform : BrickTransforms)
-	{
-		AddBoxEdges(Transform.TransformPosition(BrickMin), Transform.TransformPosition(BrickMax));
-	}
+	FFractureEngineFracturing::GenerateBrickTransforms(Bounds,
+		BrickTransforms,
+		static_cast<EFractureBrickBondEnum>(BrickSettings->Bond),
+		BrickSettings->BrickLength, BrickSettings->BrickHeight, BrickSettings->BrickDepth,
+		Edges);
 }
 
 void UFractureToolBrick::UpdateBrickTransforms()
diff --git a/Engine/Plugins/Experimental/Fracture/Source/FractureEngine/Private/FractureEngineFracturing.cpp b/Engine/Plugins/Experimental/Fracture/Source/FractureEngine/Private/FractureEngineFracturing.cpp
--- a/Engine/Plugins/Experimental/Fracture/Source/FractureEngine/Private/FractureEngineFracturing.cpp
+++ b/Engine/Plugins/Experimental/Fracture/Source/FractureEngine/Private/FractureEngineFracturing.cpp
@@ -846,13 +846,25 @@ void FFractureEngineFracturing::GenerateBrickTransforms(const FBox& InBounds,
 	const float InBrickDepth,
 	TArray<TTuple<FVector, FVector>>& InOutEdges)
 {
-	const FVector& Min = InBounds.Min;
-	const FVector& Max = InBounds.Max;
-	const FVector Extents(InBounds.Max - InBounds.Min);
-
 	// Determine brick dimensions (length, depth, height) and make sure we do not exceed the limit for the number of bricks.
 	// If we would simply use the input dimensions, we are prone to running out of memory and/or exceeding the storage capabilities of TArray, and crash.
-	const FVector BrickDimensions = FractureToolBrickLocals::GetBrickDimensions(InBrickLength, InBrickHeight, InBrickDepth, Extents);
+	const FVector BrickDimensions = FractureToolBrickLocals::GetBrickDimensions(InBrickLength, InBrickHeight, InBrickDepth, InBounds.Max - InBounds.Min);
+	const FVector BrickHalfDimensions(BrickDimensions * 0.5);
+
+	// make conservative version of dims to account for 90-degree brick rotation in some patterns
+	FVector ConservativeHalfBrick = BrickHalfDimensions;
+	if (InBond != EFractureBrickBondEnum::Dataflow_FractureBrickBond_Stretcher && InBond != EFractureBrickBondEnum::Dataflow_FractureBrickBond_Stack)
+	{
+		double MaxXY = FMath::Max(ConservativeHalfBrick.X, ConservativeHalfBrick.Y);
+		ConservativeHalfBrick.X = MaxXY;
+		ConservativeHalfBrick.Y = MaxXY;
+	}
+
+	const FVector Min = InBounds.Min;
+	// Bricks are tiled up to the point where their center is beyond Max in any dimension, so
+	// to cover InBounds, we need this Max to be at least a half-brick beyond the InBounds.Max
+	const FVector Max = InBounds.Max + ConservativeHalfBrick + FVector(UE_DOUBLE_KINDA_SMALL_NUMBER);
+	const FVector Extents(Max - Min);
 
 	// Early out if we have inputs we cannot deal with.
 	if (BrickDimensions == FVector::ZeroVector)
@@ -863,8 +875,7 @@ void FFractureEngineFracturing::GenerateBrickTransforms(const FBox& InBounds,
 	// Reserve correct amount of memory to avoid re-allocations.
 	InOutBrickTransforms.Reserve(FractureToolBrickLocals::CalculateNumBricks(BrickDimensions, Extents));
 
-	const FVector BrickHalfDimensions(BrickDimensions * 0.5);
-	const FQuat HeaderRotation(FVector::UpVector, 1.5708);
+	const FQuat HeaderRotation(FVector::UpVector, UE_DOUBLE_HALF_PI);
 
 	if (InBond == EFractureBrickBondEnum::Dataflow_FractureBrickBond_Stretcher)
 	{
@@ -1041,8 +1052,6 @@ int32 FFractureEngineFracturing::BrickCutter(FManagedArrayCollection& InOutColle
 		const FVector BrickDimensions = FractureToolBrickLocals::GetBrickDimensions(InBrickLength, InBrickHeight, InBrickDepth, Bounds.Max - Bounds.Min);
 		const FVector BrickHalfDimensions(BrickDimensions * 0.5);
 
-		const FQuat HeaderRotation(FVector::UpVector, 1.5708);
-
 		TArray<FBox> BricksToCut;
 
 		// space the bricks by the grout setting, constrained to not erase the bricks
