Fracture editor : when creating a geometry collection , make sure we can clear the dataflow or physical material entries

diff --git a/Engine/Plugins/Experimental/ChaosEditor/Source/FractureEditor/Private/FractureToolGenerators.cpp b/Engine/Plugins/Experimental/ChaosEditor/Source/FractureEditor/Private/FractureToolGenerators.cpp
--- a/Engine/Plugins/Experimental/ChaosEditor/Source/FractureEditor/Private/FractureToolGenerators.cpp
+++ b/Engine/Plugins/Experimental/ChaosEditor/Source/FractureEditor/Private/FractureToolGenerators.cpp
@@ -145,10 +145,7 @@ private:
 	}
 	void OnDataflowAssetChanged(const FAssetData& InAssetData)
 	{
-		if (UDataflow* NewDataflow = Cast<UDataflow>(InAssetData.GetAsset()))
-		{
-			DataflowAsset = NewDataflow;
-		}
+		DataflowAsset = Cast<UDataflow>(InAssetData.GetAsset());
 	}
 
 	FString GetPhysicalMaterialPath() const
@@ -161,10 +158,7 @@ private:
 	}
 	void OnPhysicalMaterialChanged(const FAssetData& InAssetData)
 	{
-		if (UPhysicalMaterial* NewPhysicalMaterial = Cast<UPhysicalMaterial>(InAssetData.GetAsset()))
-		{
-			PhysicalMaterial = NewPhysicalMaterial;
-		}
+		PhysicalMaterial = Cast<UPhysicalMaterial>(InAssetData.GetAsset());
 	}
 	
 
