Fracture editor add a None option to showing selected element in a geometry collection

diff --git a/Engine/Plugins/Experimental/ChaosEditor/Source/FractureEditor/Private/FractureEditorMode.cpp b/Engine/Plugins/Experimental/ChaosEditor/Source/FractureEditor/Private/FractureEditorMode.cpp
--- a/Engine/Plugins/Experimental/ChaosEditor/Source/FractureEditor/Private/FractureEditorMode.cpp
+++ b/Engine/Plugins/Experimental/ChaosEditor/Source/FractureEditor/Private/FractureEditorMode.cpp
@@ -260,10 +260,16 @@ void UFractureEditorMode::DrawHUD(FEditorViewportClient* ViewportClient, FViewpo
 
 void UFractureEditorMode::RenderSelectionAsBoxes(const FSceneView* View, FViewport* Viewport, FPrimitiveDrawInterface* PDI)
 {
-	// this is a uggestion has nanite based geometry collection may want to override it 
+	// this is a suggestion has nanite based geometry collection may want to override it 
 	bool bShouldDisplayBounds = false;
 	if (const UFractureSettings* FractureSettings = GetDefault<UFractureSettings>())
 	{
+		// early exit if the None option has been selected
+		if (FractureSettings->SelectionDisplayMode == EFractureSelectionDisplayMode::None)
+		{
+			return;
+		}
+
 		bShouldDisplayBounds = (FractureSettings->SelectionDisplayMode == EFractureSelectionDisplayMode::BoundingBox);
 	}
 
diff --git a/Engine/Plugins/Experimental/ChaosEditor/Source/FractureEditor/Public/FractureSettings.h b/Engine/Plugins/Experimental/ChaosEditor/Source/FractureEditor/Public/FractureSettings.h
--- a/Engine/Plugins/Experimental/ChaosEditor/Source/FractureEditor/Public/FractureSettings.h
+++ b/Engine/Plugins/Experimental/ChaosEditor/Source/FractureEditor/Public/FractureSettings.h
@@ -17,6 +17,9 @@ enum class EFractureSelectionDisplayMode : uint8
 
 	/** Use bounding boxes to show selection - Nanite Geometry Collection will default to that */
 	BoundingBox,
+
+	/** do not display bone selection at all  */
+	None,
 };
 
 
