add a bake-specific 'padding' option to the fracture auto UV tool, to copy boundary pixel values out to a larger range of empty pixels when baking (without affecting the spacing of the UV layout step)

#Synced-CL 50262759

diff --git a/Engine/Plugins/Experimental/ChaosEditor/Source/FractureEditor/Private/FractureToolUV.cpp b/Engine/Plugins/Experimental/ChaosEditor/Source/FractureEditor/Private/FractureToolUV.cpp
--- a/Engine/Plugins/Experimental/ChaosEditor/Source/FractureEditor/Private/FractureToolUV.cpp
+++ b/Engine/Plugins/Experimental/ChaosEditor/Source/FractureEditor/Private/FractureToolUV.cpp
@@ -559,7 +559,7 @@ void UFractureToolAutoUV::BakeTextureForComponent(UGeometryCollectionComponent*
 	BakeOp->MaterialIDs = !TargetsMaterialIDs(AutoUVSettings->TargetFaces) ? EmptyMaterialIDs : AutoUVSettings->MaterialIDs;
 
 	BakeOp->UVLayer = AutoUVSettings->GetSelectedChannelIndex();
-	BakeOp->GutterSize = AutoUVSettings->GutterSize;
+	BakeOp->GutterSize = AutoUVSettings->GutterSize + AutoUVSettings->GutterPaddingOnBake;
 	BakeOp->TargetFaces = ConvertTargetFacesEnum(AutoUVSettings->TargetFaces);
 
 	int32 OutputRes = (int32)AutoUVSettings->Resolution;
diff --git a/Engine/Plugins/Experimental/ChaosEditor/Source/FractureEditor/Private/FractureToolUV.h b/Engine/Plugins/Experimental/ChaosEditor/Source/FractureEditor/Private/FractureToolUV.h
--- a/Engine/Plugins/Experimental/ChaosEditor/Source/FractureEditor/Private/FractureToolUV.h
+++ b/Engine/Plugins/Experimental/ChaosEditor/Source/FractureEditor/Private/FractureToolUV.h
@@ -154,6 +154,10 @@ public:
 	UPROPERTY(EditAnywhere, Category = MapSettings)
 	bool bClearGutterAlpha = false;
 
+	/** Amount to extend copied border pixel values into empty texture space, beyond the Gutter Size. Only applied when baking, does not affect spacing for UV Layout. */
+	UPROPERTY(EditAnywhere, Category = MapSettings, meta = (ClampMin = "0", UIMax = "10"))
+	int32 GutterPaddingOnBake = 4;
+
 	/** The resulting automatically-generated texture map */
 	UPROPERTY(VisibleAnywhere, Category = MapSettings)
 	TObjectPtr<UTexture2D> Result;
