add 'full image' padding option to fracture autouv baking

#Synced-CL 50388552

diff --git a/Engine/Plugins/Experimental/ChaosEditor/Source/FractureEditor/Private/FractureToolUV.cpp b/Engine/Plugins/Experimental/ChaosEditor/Source/FractureEditor/Private/FractureToolUV.cpp
--- a/Engine/Plugins/Experimental/ChaosEditor/Source/FractureEditor/Private/FractureToolUV.cpp
+++ b/Engine/Plugins/Experimental/ChaosEditor/Source/FractureEditor/Private/FractureToolUV.cpp
@@ -558,11 +558,11 @@ void UFractureToolAutoUV::BakeTextureForComponent(UGeometryCollectionComponent*
 	TArray<int32> EmptyMaterialIDs;
 	BakeOp->MaterialIDs = !TargetsMaterialIDs(AutoUVSettings->TargetFaces) ? EmptyMaterialIDs : AutoUVSettings->MaterialIDs;
 
+	int32 OutputRes = (int32)AutoUVSettings->Resolution;
 	BakeOp->UVLayer = AutoUVSettings->GetSelectedChannelIndex();
-	BakeOp->GutterSize = AutoUVSettings->GutterSize + AutoUVSettings->GutterPaddingOnBake;
+	BakeOp->GutterSize = AutoUVSettings->GutterSize + (AutoUVSettings->bFullyPadGutterOnBake ? OutputRes : AutoUVSettings->GutterPaddingOnBake);
 	BakeOp->TargetFaces = ConvertTargetFacesEnum(AutoUVSettings->TargetFaces);
 
-	int32 OutputRes = (int32)AutoUVSettings->Resolution;
 	FImageDimensions Dimensions(OutputRes, OutputRes);
 	BakeOp->ImageBuilder = MakeUnique<FImageType>();
 	BakeOp->ImageBuilder->SetDimensions(Dimensions);
diff --git a/Engine/Plugins/Experimental/ChaosEditor/Source/FractureEditor/Private/FractureToolUV.h b/Engine/Plugins/Experimental/ChaosEditor/Source/FractureEditor/Private/FractureToolUV.h
--- a/Engine/Plugins/Experimental/ChaosEditor/Source/FractureEditor/Private/FractureToolUV.h
+++ b/Engine/Plugins/Experimental/ChaosEditor/Source/FractureEditor/Private/FractureToolUV.h
@@ -150,12 +150,16 @@ public:
 	UPROPERTY(EditAnywhere, Category = MapSettings, meta = (ClampMin = "1", ClampMax = "10", UIMax = "6"))
 	int32 GutterSize = 2;
 
-	/** Whether to set gutter pixels to transparent, or to copy them unchanged from the closest interior pixel value. */
+	/** Whether to set gutter texels to transparent, or to copy them unchanged from the closest interior texel value. */
 	UPROPERTY(EditAnywhere, Category = MapSettings)
 	bool bClearGutterAlpha = false;
 
-	/** Amount to extend copied border pixel values into empty texture space, beyond the Gutter Size. Only applied when baking, does not affect spacing for UV Layout. */
-	UPROPERTY(EditAnywhere, Category = MapSettings, meta = (ClampMin = "0", UIMax = "10"))
+	/** Whether to extend copied border texel values into all empty texture space, beyond the Gutter Size. Only applied when baking, does not affect spacing for UV Layout. */
+	UPROPERTY(EditAnywhere, Category = MapSettings)
+	bool bFullyPadGutterOnBake = true;
+
+	/** Amount of texels to extend copied border texel values into empty texture space, beyond the Gutter Size. Only applied when baking, does not affect spacing for UV Layout. */
+	UPROPERTY(EditAnywhere, Category = MapSettings, meta = (ClampMin = "0", UIMax = "10", EditCondition="!bFullyPadGutterOnBake"))
 	int32 GutterPaddingOnBake = 4;
 
 	/** The resulting automatically-generated texture map */
diff --git a/Engine/Plugins/Experimental/GeometryCollectionPlugin/Source/GeometryCollectionNodes/Private/Dataflow/GeometryCollectionTextureNodes.cpp b/Engine/Plugins/Experimental/GeometryCollectionPlugin/Source/GeometryCollectionNodes/Private/Dataflow/GeometryCollectionTextureNodes.cpp
--- a/Engine/Plugins/Experimental/GeometryCollectionPlugin/Source/GeometryCollectionNodes/Private/Dataflow/GeometryCollectionTextureNodes.cpp
+++ b/Engine/Plugins/Experimental/GeometryCollectionPlugin/Source/GeometryCollectionNodes/Private/Dataflow/GeometryCollectionTextureNodes.cpp
@@ -35,6 +35,7 @@ FBakeTextureFromCollectionDataflowNode::FBakeTextureFromCollectionDataflowNode(c
 	RegisterInputConnection(&ThicknessFactor).SetCanHidePin(true).SetPinIsHidden(true);
 	RegisterInputConnection(&MaxCurvature).SetCanHidePin(true).SetPinIsHidden(true);
 	RegisterInputConnection(&bClearGutterAlpha).SetCanHidePin(true).SetPinIsHidden(true);
+	RegisterInputConnection(&bFullyPadGutter).SetCanHidePin(true).SetPinIsHidden(true);
 	
 	RegisterOutputConnection(&Collection, &Collection);
 	RegisterOutputConnection(&Image);
@@ -56,7 +57,8 @@ void FBakeTextureFromCollectionDataflowNode::Evaluate(UE::Dataflow::FContext& Co
 		FManagedArrayCollection InCollection = GetValue(Context, &Collection);
 		const int32 InUVChannel = GetValue(Context, &UVChannel);
 		const int32 InResolution = (int32)GetValue(Context, &Resolution);
-		const int32 InGutterSize = GetValue(Context, &GutterSize);
+		const int32 InGutterSize = GetValue(Context, &bFullyPadGutter) ? InResolution : GetValue(Context, &GutterSize);
+		
 
 		TUniquePtr<FGeometryCollection> GeomCollection = TUniquePtr<FGeometryCollection>(InCollection.NewCopy<FGeometryCollection>());
 		check(GeomCollection.IsValid());
diff --git a/Engine/Plugins/Experimental/GeometryCollectionPlugin/Source/GeometryCollectionNodes/Public/Dataflow/GeometryCollectionTextureNodes.h b/Engine/Plugins/Experimental/GeometryCollectionPlugin/Source/GeometryCollectionNodes/Public/Dataflow/GeometryCollectionTextureNodes.h
--- a/Engine/Plugins/Experimental/GeometryCollectionPlugin/Source/GeometryCollectionNodes/Public/Dataflow/GeometryCollectionTextureNodes.h
+++ b/Engine/Plugins/Experimental/GeometryCollectionPlugin/Source/GeometryCollectionNodes/Public/Dataflow/GeometryCollectionTextureNodes.h
@@ -73,8 +73,12 @@ private:
 	UPROPERTY(EditAnywhere, Category = ImageOptions, meta = (DataflowInput))
 	EDataflowImageResolution Resolution = EDataflowImageResolution::Resolution512;
 
-	/** Approximate space to leave between UV islands, measured in texels */
-	UPROPERTY(EditAnywhere, Category = ChannelOptions, meta = (DataflowInput, ClampMin = 1, UIMin = 1))
+	/** Whether to pad out gutter size so that all empty texels are filled */
+	UPROPERTY(EditAnywhere, Category = ChannelOptions, meta = (DataflowInput))
+	bool bFullyPadGutter = false;
+
+	/** Approximate space to fill in between UV islands, measured in texels. Ignored if Fully Pad Gutter is enabled. */
+	UPROPERTY(EditAnywhere, Category = ChannelOptions, meta = (DataflowInput, EditCondition="!bFullyPadGutter", ClampMin = 1, UIMin = 1))
 	int32 GutterSize = 1;
 
 	// If true, will clear the alpha channel for gutter pixels, so they are transparent and the UV boundaries are clearer. If false, gutter pixel values will be copied as-is from the nearest interior pixel.
