[Fracture|ITF Gizmos] Opt in to the new gizmo pipeline for the fracture editor mode.

#rnx
#Synced-CL 51471837

diff --git a/Engine/Plugins/Experimental/ChaosEditor/Source/FractureEditor/Public/FractureEditorMode.h b/Engine/Plugins/Experimental/ChaosEditor/Source/FractureEditor/Public/FractureEditorMode.h
--- a/Engine/Plugins/Experimental/ChaosEditor/Source/FractureEditor/Public/FractureEditorMode.h
+++ b/Engine/Plugins/Experimental/ChaosEditor/Source/FractureEditor/Public/FractureEditorMode.h
@@ -52,6 +52,7 @@ public:
 
 	// ILegacyEdModeViewportInterface
 	virtual bool InputAxis(FEditorViewportClient* InViewportClient, FViewport* InViewport, int32 ControllerId, FKey Key, float Delta, float DeltaTime) override;
+	virtual bool RequiresLegacyViewportInteractions() const override { return false; }
 
 	// ILegacyEdModeSelectInterface
 	virtual bool BoxSelect(FBox& InBox, bool InSelect = true) override;
