Add new split island controls to fracture tools

#Synced-CL 52268326

diff --git a/Engine/Plugins/Experimental/ChaosEditor/Source/FractureEditor/Private/FractureToolBrick.cpp b/Engine/Plugins/Experimental/ChaosEditor/Source/FractureEditor/Private/FractureToolBrick.cpp
--- a/Engine/Plugins/Experimental/ChaosEditor/Source/FractureEditor/Private/FractureToolBrick.cpp
+++ b/Engine/Plugins/Experimental/ChaosEditor/Source/FractureEditor/Private/FractureToolBrick.cpp
@@ -252,6 +252,8 @@ public:
 	int Seed;
 	FTransform Transform;
 	bool bSplitIslands = true;
+	double CloseVertexDistance = 1e-3;
+	double VertexToSurfaceBridgeDistance = 0;
 
 	// TGenericDataOperator interface:
 	virtual void CalculateResult(FProgressCancel* Progress) override
@@ -261,7 +263,7 @@ public:
 			return;
 		}
 
-		ResultGeometryIndex = CutMultipleWithPlanarCells(Cells, *CollectionCopy, Selection, Grout, PointSpacing, Seed, Transform, true, true, Progress, CellsOrigin, bSplitIslands);
+		ResultGeometryIndex = CutMultipleWithPlanarCells(Cells, *CollectionCopy, Selection, Grout, PointSpacing, Seed, Transform, true, true, Progress, CellsOrigin, FIslandSplitSettings(bSplitIslands, CloseVertexDistance, VertexToSurfaceBridgeDistance));
 
 		SetResult(MoveTemp(CollectionCopy));
 	}
@@ -307,6 +309,8 @@ int32 UFractureToolBrick::ExecuteFracture(const FFractureToolContext& FractureCo
 		BrickOp->Grout = 0; // CutterSettings->Grout; // Note: Grout is currently baked directly into the brick cells above
 		BrickOp->PointSpacing = CollisionSettings->GetPointSpacing();
 		BrickOp->bSplitIslands = CutterSettings->bSplitIslands;
+		BrickOp->CloseVertexDistance = CutterSettings->CloseVertexDistance;
+		BrickOp->VertexToSurfaceBridgeDistance = CutterSettings->VertexToSurfaceBridgeDistance;
 		const bool bBricksAreTouching = CutterSettings->Grout <= UE_KINDA_SMALL_NUMBER;
 		BrickOp->Cells = FPlanarCells(BricksToCut, bBricksAreTouching);
 		if (CutterSettings->Amplitude > 0.0f)
diff --git a/Engine/Plugins/Experimental/ChaosEditor/Source/FractureEditor/Private/FractureToolCutter.cpp b/Engine/Plugins/Experimental/ChaosEditor/Source/FractureEditor/Private/FractureToolCutter.cpp
--- a/Engine/Plugins/Experimental/ChaosEditor/Source/FractureEditor/Private/FractureToolCutter.cpp
+++ b/Engine/Plugins/Experimental/ChaosEditor/Source/FractureEditor/Private/FractureToolCutter.cpp
@@ -626,6 +626,8 @@ public:
 	int Seed;
 	FTransform Transform;
 	bool bSplitIslands = true;
+	double CloseVertexDistance = 1e-3;
+	double VertexToSurfaceBridgeDistance = 0;
 
 	// TGenericDataOperator interface:
 	virtual void CalculateResult(FProgressCancel* Progress) override
@@ -657,7 +659,7 @@ public:
 		FProgressCancel::FProgressScope FractureMeshProgress =
 			FProgressCancel::CreateScopeTo(Progress, 1, LOCTEXT("FractureMeshMessage", "Fracturing Mesh"));
 
-		ResultGeometryIndex = CutMultipleWithPlanarCells(VoronoiPlanarCells, *CollectionCopy, Selection, Grout, PointSpacing, Seed, Transform, true, true, Progress, Origin, bSplitIslands);
+		ResultGeometryIndex = CutMultipleWithPlanarCells(VoronoiPlanarCells, *CollectionCopy, Selection, Grout, PointSpacing, Seed, Transform, true, true, Progress, Origin, FIslandSplitSettings(bSplitIslands, CloseVertexDistance, VertexToSurfaceBridgeDistance));
 		
 		SetResult(MoveTemp(CollectionCopy));
 	}
@@ -679,6 +681,8 @@ int32 UFractureToolVoronoiCutterBase::ExecuteFracture(const FFractureToolContext
 		VoronoiOp->PointSpacing = CollisionSettings->GetPointSpacing();
 		VoronoiOp->Sites = Sites;
 		VoronoiOp->bSplitIslands = CutterSettings->bSplitIslands;
+		VoronoiOp->CloseVertexDistance = CutterSettings->CloseVertexDistance;
+		VoronoiOp->VertexToSurfaceBridgeDistance = CutterSettings->VertexToSurfaceBridgeDistance;
 		if (CutterSettings->Amplitude > 0.0f)
 		{
 			FNoiseSettings Settings;
diff --git a/Engine/Plugins/Experimental/ChaosEditor/Source/FractureEditor/Private/FractureToolCutter.h b/Engine/Plugins/Experimental/ChaosEditor/Source/FractureEditor/Private/FractureToolCutter.h
--- a/Engine/Plugins/Experimental/ChaosEditor/Source/FractureEditor/Private/FractureToolCutter.h
+++ b/Engine/Plugins/Experimental/ChaosEditor/Source/FractureEditor/Private/FractureToolCutter.h
@@ -82,10 +82,18 @@ public:
 	UPROPERTY()
 	bool bGroupFractureToggleEnabled = true;
 
-	// Whether to split the fractured mesh pieces based on geometric connectivity after fracturing
-	UPROPERTY(EditAnywhere, Category = CommonFracture)
+	/** Whether to split the fractured mesh pieces into separate islands based on geometric connectivity after fracturing */
+	UPROPERTY(EditAnywhere, Category = "CommonFracture|Split Islands")
 	bool bSplitIslands = true;
 
+	/** Distance threshold for connecting spatially-close vertices before computing islands */
+	UPROPERTY(EditAnywhere, Category = "CommonFracture|Split Islands", meta = (EditCondition = "bSplitIslands", Units = "cm", ClampMin = "0.0"))
+	float CloseVertexDistance = 0.001f;
+
+	/** If > 0, bridge separate islands whose surfaces are within this vertex-to-triangle distance. 0 = disabled. */
+	UPROPERTY(EditAnywhere, Category = "CommonFracture|Split Islands", meta = (EditCondition = "bSplitIslands", Units = "cm", ClampMin = "0.0"))
+	float VertexToSurfaceBridgeDistance = 0.f;
+
 	/** Amount of space to leave between cut pieces */
 	UPROPERTY(EditAnywhere, Category = CommonFracture, meta = (UIMin = "0.0", ClampMin = "0.0", EditCondition = "bGroutSettingEnabled", HideEditConditionToggle, EditConditionHides))
 	float Grout = 0.0f;
diff --git a/Engine/Plugins/Experimental/ChaosEditor/Source/FractureEditor/Private/FractureToolMeshCut.cpp b/Engine/Plugins/Experimental/ChaosEditor/Source/FractureEditor/Private/FractureToolMeshCut.cpp
--- a/Engine/Plugins/Experimental/ChaosEditor/Source/FractureEditor/Private/FractureToolMeshCut.cpp
+++ b/Engine/Plugins/Experimental/ChaosEditor/Source/FractureEditor/Private/FractureToolMeshCut.cpp
@@ -218,6 +218,8 @@ public:
 	FTransform Transform;
 	UE::Geometry::FDynamicMesh3 CuttingMesh;
 	bool bSplitIslands = true;
+	double CloseVertexDistance = 1e-3;
+	double VertexToSurfaceBridgeDistance = 0;
 
 	// TGenericDataOperator interface:
 	virtual void CalculateResult(FProgressCancel* Progress) override
@@ -240,7 +242,7 @@ public:
 
 			FProgressCancel::FProgressScope ProgressScope(Progress, ProgressFrac);
 			constexpr bool bSetDefaultInternalMaterialsFromCollection = true;
-			int32 Index = CutWithMesh(CuttingMesh, ScatterTransform, InternalSurfaceMaterials, *CollectionCopy, BonesToCut, PointSpacing, Transform, bSetDefaultInternalMaterialsFromCollection, Progress, bSplitIslands);
+			int32 Index = CutWithMesh(CuttingMesh, ScatterTransform, InternalSurfaceMaterials, *CollectionCopy, BonesToCut, PointSpacing, Transform, bSetDefaultInternalMaterialsFromCollection, Progress, FIslandSplitSettings(bSplitIslands, CloseVertexDistance, VertexToSurfaceBridgeDistance));
 			if (Progress && Progress->Cancelled())
 			{
 				return;
@@ -318,6 +320,8 @@ int32 UFractureToolMeshCut::ExecuteFracture(const FFractureToolContext& Fracture
 		MeshCutOp->Seed = FractureContext.GetSeed();
 		MeshCutOp->Transform = FractureContext.GetTransform();
 		MeshCutOp->bSplitIslands = CutterSettings->bSplitIslands;
+		MeshCutOp->CloseVertexDistance = CutterSettings->CloseVertexDistance;
+		MeshCutOp->VertexToSurfaceBridgeDistance = CutterSettings->VertexToSurfaceBridgeDistance;
 		if (LocalCutSettings->CutDistribution == EMeshCutDistribution::SingleCut)
 		{
 			MeshCutOp->MeshTransforms.Add(LocalCutSettings->CuttingActor->GetTransform());
diff --git a/Engine/Plugins/Experimental/ChaosEditor/Source/FractureEditor/Private/FractureToolPlaneCut.cpp b/Engine/Plugins/Experimental/ChaosEditor/Source/FractureEditor/Private/FractureToolPlaneCut.cpp
--- a/Engine/Plugins/Experimental/ChaosEditor/Source/FractureEditor/Private/FractureToolPlaneCut.cpp
+++ b/Engine/Plugins/Experimental/ChaosEditor/Source/FractureEditor/Private/FractureToolPlaneCut.cpp
@@ -377,11 +377,13 @@ public:
 	FTransform Transform;
 	UE::Geometry::FDynamicMesh3 CuttingMesh;
 	bool bSplitIslands = true;
+	double CloseVertexDistance = 1e-3;
+	double VertexToSurfaceBridgeDistance = 0;
 
 	// TGenericDataOperator interface:
 	virtual void CalculateResult(FProgressCancel* Progress) override
 	{
-		ResultGeometryIndex = CutMultipleWithMultiplePlanes(CuttingPlanes, InternalSurfaceMaterials, *CollectionCopy, Selection, Grout, PointSpacing, Seed, Transform, true, Progress, bSplitIslands);
+		ResultGeometryIndex = CutMultipleWithMultiplePlanes(CuttingPlanes, InternalSurfaceMaterials, *CollectionCopy, Selection, Grout, PointSpacing, Seed, Transform, true, Progress, FIslandSplitSettings(bSplitIslands, CloseVertexDistance, VertexToSurfaceBridgeDistance));
 
 		if (Progress && Progress->Cancelled())
 		{
@@ -403,6 +405,8 @@ int32 UFractureToolPlaneCut::ExecuteFracture(const FFractureToolContext& Fractur
 		PlaneCutOp->Seed = FractureContext.GetSeed();
 		PlaneCutOp->Transform = FractureContext.GetTransform();
 		PlaneCutOp->bSplitIslands = CutterSettings->bSplitIslands;
+		PlaneCutOp->CloseVertexDistance = CutterSettings->CloseVertexDistance;
+		PlaneCutOp->VertexToSurfaceBridgeDistance = CutterSettings->VertexToSurfaceBridgeDistance;
 
 		if (GizmoSettings->IsGizmoEnabled())
 		{
diff --git a/Engine/Plugins/Experimental/ChaosEditor/Source/FractureEditor/Private/FractureToolSlice.cpp b/Engine/Plugins/Experimental/ChaosEditor/Source/FractureEditor/Private/FractureToolSlice.cpp
--- a/Engine/Plugins/Experimental/ChaosEditor/Source/FractureEditor/Private/FractureToolSlice.cpp
+++ b/Engine/Plugins/Experimental/ChaosEditor/Source/FractureEditor/Private/FractureToolSlice.cpp
@@ -185,7 +185,7 @@ int32 UFractureToolSlice::ExecuteFracture(const FFractureToolContext& FractureCo
 		// Proximity is invalidated.
 		ClearProximity(FractureContext.GetGeometryCollection().Get());
 
-		return CutMultipleWithMultiplePlanes(CuttingPlanes, InternalSurfaceMaterials, *FractureContext.GetGeometryCollection(), FractureContext.GetSelection(), CutterSettings->Grout, CollisionSettings->GetPointSpacing(), FractureContext.GetSeed(), FractureContext.GetTransform(), true, nullptr, CutterSettings->bSplitIslands);
+		return CutMultipleWithMultiplePlanes(CuttingPlanes, InternalSurfaceMaterials, *FractureContext.GetGeometryCollection(), FractureContext.GetSelection(), CutterSettings->Grout, CollisionSettings->GetPointSpacing(), FractureContext.GetSeed(), FractureContext.GetTransform(), true, nullptr, FIslandSplitSettings(CutterSettings->bSplitIslands, CutterSettings->CloseVertexDistance, CutterSettings->VertexToSurfaceBridgeDistance));
 	}
 
 	return INDEX_NONE;
