TLazyObjectPtr Deprecation pt 3:
Change usage of TLazyObjectPtr used in config objects, serialization for both use the same export text flow, meaning no settings should be lost

#Synced-CL 52252558

diff --git a/Engine/Plugins/Bridge/Source/MegascansPlugin/Public/MSSettings.h b/Engine/Plugins/Bridge/Source/MegascansPlugin/Public/MSSettings.h
--- a/Engine/Plugins/Bridge/Source/MegascansPlugin/Public/MSSettings.h
+++ b/Engine/Plugins/Bridge/Source/MegascansPlugin/Public/MSSettings.h
@@ -82,16 +82,16 @@ class UMaterialPresetsSettings
 public:
 	/** Replace default master material with your own custom master material for all 3D assets. Default material is used if field is left empty. */
 	UPROPERTY(Transient, DisplayName = "3D Master Material", EditAnywhere, Category = "MasterMaterialOverrides")		
-		TLazyObjectPtr <class UMaterial> MasterMaterial3d;
-	//TLazyObjectPtr<class UMaterialInterface> MasterMaterial3d;
+	TSoftObjectPtr<class UMaterial> MasterMaterial3d;
+	//TSoftObjectPtr<class UMaterialInterface> MasterMaterial3d;
 
 	/** Replace default master material with your own custom master material for all Surfaces. Default material is used if field is left empty. */
 	UPROPERTY(Transient, DisplayName = "Surface Master Material", EditAnywhere, Category = "MasterMaterialOverrides")
-		TLazyObjectPtr<class UMaterial> MasterMaterialSurface;
+	TSoftObjectPtr<class UMaterial> MasterMaterialSurface;
 
 	/** Replace default master material with your own custom master material for all 3D Plants. Default material is used if field is left empty. */
 	UPROPERTY(Transient, DisplayName = "Plant Master Material", EditAnywhere, Category = "MasterMaterialOverrides")
-		TLazyObjectPtr<class UMaterial> MasterMaterialPlant;
+	TSoftObjectPtr<class UMaterial> MasterMaterialPlant;
 
 #if WITH_EDITOR
 	virtual void PostEditChangeProperty(FPropertyChangedEvent& PropertyThatWillChange) override;
diff --git a/Engine/Plugins/Experimental/ChaosEditor/Source/FractureEditor/Private/FractureToolCustomVoronoi.cpp b/Engine/Plugins/Experimental/ChaosEditor/Source/FractureEditor/Private/FractureToolCustomVoronoi.cpp
--- a/Engine/Plugins/Experimental/ChaosEditor/Source/FractureEditor/Private/FractureToolCustomVoronoi.cpp
+++ b/Engine/Plugins/Experimental/ChaosEditor/Source/FractureEditor/Private/FractureToolCustomVoronoi.cpp
@@ -167,13 +167,13 @@ void UFractureToolCustomVoronoi::GenerateLivePattern(const TArray<FFractureToolC
 	}
 	else if (CustomVoronoiSettings->VoronoiPattern == EVoronoiPattern::MeshVertices)
 	{
-		auto HasMesh = [this]() -> bool
+		auto HasMesh = [](AStaticMeshActor* ReferenceMesh) -> bool
 		{
-			if (CustomVoronoiSettings->ReferenceMesh == nullptr)
+			if (ReferenceMesh == nullptr)
 			{
 				return false;
 			}
-			const UStaticMeshComponent* Component = CustomVoronoiSettings->ReferenceMesh->GetStaticMeshComponent();
+			const UStaticMeshComponent* Component = ReferenceMesh->GetStaticMeshComponent();
 			if (Component == nullptr)
 			{
 				return false;
@@ -189,15 +189,16 @@ void UFractureToolCustomVoronoi::GenerateLivePattern(const TArray<FFractureToolC
 			}
 			return true;
 		};
-		if (HasMesh())
+		AStaticMeshActor* ReferenceMesh = CustomVoronoiSettings->ReferenceMesh.Get();
+		if (HasMesh(ReferenceMesh))
 		{
-			FTransform VerticesTransform = CustomVoronoiSettings->ReferenceMesh->GetStaticMeshComponent()->GetComponentTransform();
+			FTransform VerticesTransform = ReferenceMesh->GetStaticMeshComponent()->GetComponentTransform();
 			if (!CustomVoronoiSettings->bStartAtActor)
 			{
 				VerticesTransform.SetLocation(GizmoTransform.GetLocation());
 			}
 
-			FMeshDescription* MeshDescription = CustomVoronoiSettings->ReferenceMesh->GetStaticMeshComponent()->GetStaticMesh()->GetMeshDescription(0);
+			FMeshDescription* MeshDescription = ReferenceMesh->GetStaticMeshComponent()->GetStaticMesh()->GetMeshDescription(0);
 			TVertexAttributesConstRef<FVector3f> VertexPositions = MeshDescription->GetVertexPositions();
 			TArray<FVector3f> Normals; // Per-vertex normals for offsets
 			float NormalOffset = CustomVoronoiSettings->NormalOffset;
diff --git a/Engine/Plugins/Experimental/ChaosEditor/Source/FractureEditor/Private/FractureToolCustomVoronoi.h b/Engine/Plugins/Experimental/ChaosEditor/Source/FractureEditor/Private/FractureToolCustomVoronoi.h
--- a/Engine/Plugins/Experimental/ChaosEditor/Source/FractureEditor/Private/FractureToolCustomVoronoi.h
+++ b/Engine/Plugins/Experimental/ChaosEditor/Source/FractureEditor/Private/FractureToolCustomVoronoi.h
@@ -89,7 +89,7 @@ public:
 	//~ TODO: Have reference mesh fall back to the input fracture mesh if not provided here
 	/** Static mesh actor to be used as a reference when adding Voronoi sites */
 	UPROPERTY(EditAnywhere, Category = LiveVoronoiSites, meta = (DisplayName = "Reference Mesh", EditCondition = "VoronoiPattern == EVoronoiPattern::MeshVertices", EditConditionHides))
-	TLazyObjectPtr<AStaticMeshActor> ReferenceMesh;
+	TSoftObjectPtr<AStaticMeshActor> ReferenceMesh;
 
 	/** Whether to use the reference mesh actor's transform when placing the Voronoi sites, or center them at the current gizmo location instead */
 	UPROPERTY(EditAnywhere, Category = LiveVoronoiSites, meta = (EditCondition = "VoronoiPattern == EVoronoiPattern::MeshVertices", EditConditionHides))
diff --git a/Engine/Plugins/Experimental/ChaosEditor/Source/FractureEditor/Private/FractureToolMeshCut.cpp b/Engine/Plugins/Experimental/ChaosEditor/Source/FractureEditor/Private/FractureToolMeshCut.cpp
--- a/Engine/Plugins/Experimental/ChaosEditor/Source/FractureEditor/Private/FractureToolMeshCut.cpp
+++ b/Engine/Plugins/Experimental/ChaosEditor/Source/FractureEditor/Private/FractureToolMeshCut.cpp
@@ -63,11 +63,12 @@ TArray<UObject*> UFractureToolMeshCut::GetSettingsObjects() const
 bool UFractureToolMeshCut::IsCuttingActorValid()
 {
 	const UFractureMeshCutSettings* LocalCutSettings = MeshCutSettings;
-	if (LocalCutSettings->CuttingActor == nullptr)
+	AStaticMeshActor* CuttingActor = LocalCutSettings->CuttingActor.Get();
+	if (CuttingActor == nullptr)
 	{
 		return false;
 	}
-	const UStaticMeshComponent* Component = LocalCutSettings->CuttingActor->GetStaticMeshComponent();
+	const UStaticMeshComponent* Component = CuttingActor->GetStaticMeshComponent();
 	if (Component == nullptr)
 	{
 		return false;
@@ -88,7 +89,7 @@ void UFractureToolMeshCut::Render(const FSceneView* View, FViewport* Viewport, F
 {
 	if (CutterSettings->bDrawDiagram && IsCuttingActorValid())
 	{
-		FBox Box = MeshCutSettings->CuttingActor->GetStaticMeshComponent()->GetStaticMesh()->GetBoundingBox();
+		FBox Box = MeshCutSettings->CuttingActor.Get()->GetStaticMeshComponent()->GetStaticMesh()->GetBoundingBox();
 		
 		EnumerateVisualizationMapping(TransformsMappings, RenderMeshTransforms.Num(), [&](int32 Idx, FVector ExplodedVector)
 		{
@@ -324,13 +325,13 @@ int32 UFractureToolMeshCut::ExecuteFracture(const FFractureToolContext& Fracture
 		MeshCutOp->VertexToSurfaceBridgeDistance = CutterSettings->VertexToSurfaceBridgeDistance;
 		if (LocalCutSettings->CutDistribution == EMeshCutDistribution::SingleCut)
 		{
-			MeshCutOp->MeshTransforms.Add(LocalCutSettings->CuttingActor->GetTransform());
+			MeshCutOp->MeshTransforms.Add(LocalCutSettings->CuttingActor.Get()->GetTransform());
 		}
 		else
 		{
 			GenerateMeshTransforms(FractureContext, MeshCutOp->MeshTransforms);
 		}
-		FMeshDescription* MeshDescription = LocalCutSettings->CuttingActor->GetStaticMeshComponent()->GetStaticMesh()->GetMeshDescription(0);
+		FMeshDescription* MeshDescription = LocalCutSettings->CuttingActor.Get()->GetStaticMeshComponent()->GetStaticMesh()->GetMeshDescription(0);
 		MeshCutOp->CuttingMesh = ConvertMeshDescriptionToCuttingDynamicMesh(MeshDescription, FractureContext.GetGeometryCollection()->NumUVLayers());
 
 		int Result = RunCancellableGeometryCollectionOp<FMeshFractureOp>(*(FractureContext.GetGeometryCollection()),
diff --git a/Engine/Plugins/Experimental/ChaosEditor/Source/FractureEditor/Private/FractureToolMeshCut.h b/Engine/Plugins/Experimental/ChaosEditor/Source/FractureEditor/Private/FractureToolMeshCut.h
--- a/Engine/Plugins/Experimental/ChaosEditor/Source/FractureEditor/Private/FractureToolMeshCut.h
+++ b/Engine/Plugins/Experimental/ChaosEditor/Source/FractureEditor/Private/FractureToolMeshCut.h
@@ -31,7 +31,7 @@ public:
 
 	/** Static Mesh Actor to be used as a cutting surface. For best results, use a closed, watertight mesh */
 	UPROPERTY(EditAnywhere, Category = MeshCut, meta = (DisplayName = "Cutting Actor"))
-	TLazyObjectPtr<AStaticMeshActor> CuttingActor;
+	TSoftObjectPtr<AStaticMeshActor> CuttingActor;
 
 	/** How to arrange the mesh cuts in space */
 	UPROPERTY(EditAnywhere, Category = Distribution)
