Add new fracture mode utility to split islands

#Synced-CL 52609914

diff --git a/Engine/Plugins/Experimental/ChaosEditor/Source/FractureEditor/Private/FractureEditorCommands.h b/Engine/Plugins/Experimental/ChaosEditor/Source/FractureEditor/Private/FractureEditorCommands.h
--- a/Engine/Plugins/Experimental/ChaosEditor/Source/FractureEditor/Private/FractureEditorCommands.h
+++ b/Engine/Plugins/Experimental/ChaosEditor/Source/FractureEditor/Private/FractureEditorCommands.h
@@ -97,6 +97,7 @@ class FFractureEditorCommands : public TCommands<FFractureEditorCommands>
 		TSharedPtr< FUICommandInfo > MakeConvex;
 		TSharedPtr< FUICommandInfo > Proximity;
 		TSharedPtr< FUICommandInfo > FixTinyGeo;
+		TSharedPtr< FUICommandInfo > SplitIslands;
 
 		// Property Commands
 		TSharedPtr< FUICommandInfo > SetInitialDynamicState;
diff --git a/Engine/Plugins/Experimental/ChaosEditor/Source/FractureEditor/Private/FractureEditorModeToolkit.cpp b/Engine/Plugins/Experimental/ChaosEditor/Source/FractureEditor/Private/FractureEditorModeToolkit.cpp
--- a/Engine/Plugins/Experimental/ChaosEditor/Source/FractureEditor/Private/FractureEditorModeToolkit.cpp
+++ b/Engine/Plugins/Experimental/ChaosEditor/Source/FractureEditor/Private/FractureEditorModeToolkit.cpp
@@ -1336,6 +1336,7 @@ void FFractureEditorModeToolkit::BuildToolPalette(FName PaletteIndex, class FToo
 		ToolbarBuilder.AddToolBarButton(Commands.MakeConvex);
 		ToolbarBuilder.AddToolBarButton(Commands.Proximity);
 		ToolbarBuilder.AddToolBarButton(Commands.FixTinyGeo);
+		ToolbarBuilder.AddToolBarButton(Commands.SplitIslands);
 		ToolbarBuilder.AddToolBarButton(Commands.SetInitialDynamicState);
 		ToolbarBuilder.AddToolBarButton(Commands.SetRemoveOnBreak);
 	}
diff --git a/Engine/Plugins/Experimental/ChaosEditor/Source/FractureEditor/Private/FractureToolSplitIslands.cpp b/Engine/Plugins/Experimental/ChaosEditor/Source/FractureEditor/Private/FractureToolSplitIslands.cpp
new file mode 100644
--- /dev/null
+++ b/Engine/Plugins/Experimental/ChaosEditor/Source/FractureEditor/Private/FractureToolSplitIslands.cpp
@@ -0,0 +1,98 @@
+// Copyright Epic Games, Inc. All Rights Reserved.
+
+#include "FractureToolSplitIslands.h"
+
+#include "FractureToolCutter.h"
+#include "FractureToolContext.h"
+#include "PlanarCut.h"
+
+#include "GeometryCollection/GeometryCollectionComponent.h"
+#include UE_INLINE_GENERATED_CPP_BY_NAME(FractureToolSplitIslands)
+
+#define LOCTEXT_NAMESPACE "FractureToolSplitIslands"
+
+UFractureToolSplitIslands::UFractureToolSplitIslands(const FObjectInitializer& ObjInit)
+	: Super(ObjInit)
+{
+	SplitIslandsSettings = NewObject<UFractureSplitIslandsSettings>(GetTransientPackage(), UFractureSplitIslandsSettings::StaticClass());
+	SplitIslandsSettings->OwnerTool = this;
+	CollisionSettings = NewObject<UFractureCollisionSettings>(GetTransientPackage(), UFractureCollisionSettings::StaticClass());
+	CollisionSettings->OwnerTool = this;
+}
+
+bool UFractureToolSplitIslands::CanExecute() const
+{
+	if (!IsGeometryCollectionSelected())
+	{
+		return false;
+	}
+
+	return true;
+}
+
+FText UFractureToolSplitIslands::GetDisplayText() const
+{
+	return FText(LOCTEXT("FractureToolSplitIslands", "Split Islands"));
+}
+
+FText UFractureToolSplitIslands::GetTooltipText() const
+{
+	return FText(LOCTEXT("FractureToolSplitIslandsTooltip", "Split selected geometry into islands based on mesh connectivity and distances"));
+}
+
+FSlateIcon UFractureToolSplitIslands::GetToolIcon() const
+{
+	return FSlateIcon("FractureEditorStyle", "FractureEditor.SplitIslands");
+}
+
+void UFractureToolSplitIslands::RegisterUICommand(FFractureEditorCommands* BindingContext)
+{
+	UI_COMMAND_EXT(BindingContext, UICommandInfo, "SplitIslands", "Split", "Split selected geometry into islands based on mesh connectivity and distances", EUserInterfaceActionType::ToggleButton, FInputChord());
+	BindingContext->SplitIslands = UICommandInfo;
+}
+
+TArray<UObject*> UFractureToolSplitIslands::GetSettingsObjects() const
+{
+	TArray<UObject*> Settings;
+	Settings.Add(SplitIslandsSettings);
+	Settings.Add(CollisionSettings);
+	return Settings;
+}
+
+TArray<FFractureToolContext> UFractureToolSplitIslands::GetFractureToolContexts() const
+{
+	TArray<FFractureToolContext> Contexts;
+	TSet<UGeometryCollectionComponent*> GeomCompSelection;
+	GetSelectedGeometryCollectionComponents(GeomCompSelection);
+
+	for (UGeometryCollectionComponent* GeometryCollectionComponent : GeomCompSelection)
+	{
+		if (GeometryCollectionComponent->GetRestCollection())
+		{
+			FFractureToolContext Context(GeometryCollectionComponent);
+			Context.Sanitize();
+			Context.ConvertSelectionToLeafNodes();
+			Contexts.Add(Context);
+		}
+	}
+
+	return Contexts;
+}
+
+int32 UFractureToolSplitIslands::ExecuteFracture(const FFractureToolContext& FractureContext)
+{
+	FGeometryCollection& Collection = *FractureContext.GetGeometryCollection();
+
+	FIslandSplitSettings Settings;
+	Settings.CloseVertexDistance = SplitIslandsSettings->CloseVertexDistance;
+	Settings.VertexToSurfaceBridgeDistance = SplitIslandsSettings->VertexToSurfaceBridgeDistance;
+
+	return ::SplitIslands(
+		Collection,
+		FractureContext.GetSelection(),
+		CollisionSettings->GetPointSpacing(),
+		nullptr, // Progress
+		Settings);
+}
+
+#undef LOCTEXT_NAMESPACE
diff --git a/Engine/Plugins/Experimental/ChaosEditor/Source/FractureEditor/Private/FractureToolSplitIslands.h b/Engine/Plugins/Experimental/ChaosEditor/Source/FractureEditor/Private/FractureToolSplitIslands.h
new file mode 100644
--- /dev/null
+++ b/Engine/Plugins/Experimental/ChaosEditor/Source/FractureEditor/Private/FractureToolSplitIslands.h
@@ -0,0 +1,64 @@
+// Copyright Epic Games, Inc. All Rights Reserved.
+
+#pragma once
+
+#include "FractureTool.h"
+
+#include "FractureToolSplitIslands.generated.h"
+
+class FFractureToolContext;
+class UFractureCollisionSettings;
+
+/** Settings for the Split Islands tool */
+UCLASS(config = EditorPerProjectUserSettings)
+class UFractureSplitIslandsSettings : public UFractureToolSettings
+{
+public:
+
+	GENERATED_BODY()
+
+	UFractureSplitIslandsSettings(const FObjectInitializer& ObjInit)
+		: Super(ObjInit)
+	{}
+
+	/** Distance threshold for connecting spatially-close vertices before computing islands */
+	UPROPERTY(EditAnywhere, Category = "Split Islands", meta = (Units = "cm", ClampMin = "0.0"))
+	float CloseVertexDistance = 0.001f;
+
+	/** If > 0, bridge separate islands whose surfaces are within this vertex-to-triangle distance. 0 = disabled. */
+	UPROPERTY(EditAnywhere, Category = "Split Islands", meta = (Units = "cm", ClampMin = "0.0"))
+	float VertexToSurfaceBridgeDistance = 0.f;
+
+};
+
+
+UCLASS(DisplayName = "Split Islands", Category = "FractureTools")
+class UFractureToolSplitIslands : public UFractureModalTool
+{
+public:
+	GENERATED_BODY()
+
+	UFractureToolSplitIslands(const FObjectInitializer& ObjInit);
+
+	///
+	/// UFractureModalTool Interface
+	///
+
+	virtual FText GetDisplayText() const override;
+	virtual FText GetTooltipText() const override;
+	virtual FText GetApplyText() const override { return FText(NSLOCTEXT("FractureSplitIslands", "ExecuteSplitIslands", "Split Islands")); }
+	virtual FSlateIcon GetToolIcon() const override;
+	virtual void RegisterUICommand(FFractureEditorCommands* BindingContext) override;
+	virtual TArray<UObject*> GetSettingsObjects() const override;
+	virtual int32 ExecuteFracture(const FFractureToolContext& FractureContext) override;
+	virtual bool CanExecute() const override;
+
+	virtual TArray<FFractureToolContext> GetFractureToolContexts() const override;
+
+	UPROPERTY(EditAnywhere, Category = SplitIslands)
+	TObjectPtr<UFractureSplitIslandsSettings> SplitIslandsSettings;
+
+	UPROPERTY(EditAnywhere, Category = Collision)
+	TObjectPtr<UFractureCollisionSettings> CollisionSettings;
+
+};
